#loc add/edit CJK translation for Templates and Feature Packs

diff --git a/Templates/FP_FirstPerson/manifest.json b/Templates/FP_FirstPerson/manifest.json
--- a/Templates/FP_FirstPerson/manifest.json
+++ b/Templates/FP_FirstPerson/manifest.json
@@ -10,6 +10,10 @@
 			"Language": "ko",
 			"Text": "일인칭"
 		},
+		{
+			"Language": "zh-CN",
+			"Text": "第一人称游戏"
+		},
 		{
 			"Language": "ja",
 			"Text": "ファーストパーソン"
@@ -23,11 +27,15 @@
 		},
 		{
 			"Language": "ko",
-			"Text": "일인칭 팩은 일인칭 시점에서 보이는 (양팔로 표현되는) 플레이어 캐릭터를 특징으로 합니다. 캐릭터 이동은 키보드, 콘트롤러, 터치 디바이스의 가상 조이스틱으로 가능합니다. 추가적으로 플레이어는 마우스, 콘트롤러, 터치 디바이스의 가상 조이스틱을 사용하여 둘러보기도 가능합니다. 캐릭터에는 총이 장착되어 있어, 마우스, 콘트롤러, 터치 디바이스의 가상 조이스틱을 사용하면 레벨과 상호작용하는 '즉시 적중' 무기가 발사됩니다.\n예제 맵도 포함되어 있습니다."
+			"Text": "일인칭 시점에서 보이는 (양팔로 표현되는) 플레이어 캐릭터를 특징으로 합니다. 캐릭터 이동은 키보드, 컨트롤러, 터치 디바이스의 가상 조이스틱으로 가능합니다. 추가적으로 플레이어는 마우스, 컨트롤러, 터치 디바이스의 가상 조이스틱을 사용하여 주변을 둘러볼 수도 있습니다. 캐릭터에는 총이 장착되어 있어 마우스, 컨트롤러, 터치 디바이스의 가상 조이스틱을 사용하면 간단한 발사체가 발사되어, 아레나 벽에 튕기면서 레벨에 있는 몇몇 물리 오브젝트에 영향을 줍니다."
+		},
+		{
+			"Language": "zh-CN",
+			"Text": "第一人称模板拥有一个第一人称的玩家角色（通过一双手臂呈现）。您可以通过键盘、手柄或触摸式设备上的虚拟摇杆控制角色在关卡中四处移动。此外，您还可以使用鼠标、手柄或触摸式设备上的虚拟摇杆转动角色视角，观察四周。角色还装备了枪支，可使用鼠标、手柄或触摸式设备上的虚拟摇杆发射一颗球形子弹，并影响关卡中的一些物理对象，同时从竞技场的墙壁上弹开。"
 		},
 		{
 			"Language": "ja",
-			"Text": "このテンプレートはプレイヤーキャラクター自身の視点で操作するゲームに適しています。キャラクターは武器を持っている腕の部分だけが見えます。このキャラクターはキーボード・コントローラ・バーチャルジョイスティック・タッチパッド等を使ってレベル中を自由に歩き回ることができます。さらに、マウス・コントローラー・バーチャルジョイスティック・タッチパネル等で、視点の向きを変えることができます。装備された武器はマウス・コントローラ・バーチャルジョイスティック・タッチパネルの操作で打つことができます。\nサンプルのマップも用意されています。"
+			"Text": "ファーストパーソンには、一人称視点で表示されるプレイヤーキャラクタ（腕として表示されています）がいます。キーボード、コントローラー、タッチデバイスでのバーチャルジョイスティックでこのキャラクタを移動させることができます。追加して、プレイヤーはマウス、コントローラー、タッチデバイスでのバーチャルジョイスティックで周りを見回すことができます。さらにキャラクタは銃を装備して、マウス、コントローラー、タッチデバイスでのバーチャルジョイスティックを使ってシンプルな球体の弾を発射できます。この弾は物理オブジェクトに影響し、アリーナの壁で反射します。"
 		}
 	],
 	"AssetTypes":
@@ -35,14 +43,6 @@
 		{
 			"Language": "en",
 			"Text": "Animation Sequence, Animation Montage, Sound Wave, Skeletal Mesh, Animation Blueprint, Skeleton, Static Mesh, Texture."
-		},
-		{
-			"Language": "ko",
-			"Text": "애니메이션 시퀀스, 애니메이션 몽타주, 사운드 웨이브, 스켈레탈 메시, 애니메이션 블루프린트, 스켈레톤, 스태틱 메시, 텍스처"
-		},
-		{
-			"Language": "ja",
-			"Text": "アニメーションシーケンス・アニメーションモンタージュ・サウンドウェーブ・スケルタルメッシュ・アニメーションBP・スケルトン・スケルタルメッシュ・テクスチャ"
 		}
 	],
 	"SearchTags":
@@ -50,12 +50,8 @@
 		{
 			"Language": "en",
 			"Text": "First Person,Shoot,Fire,Aim"
-		},
-		{
-			"Language": "ko",
-			"Text": "일인칭,슈팅,발사,조준"
 		}
-	],
+	],	
 	"ClassTypes": "Character, GameMode, HUD, Actor, Character, GameMode, HUD, Actor",
 	"Category": "CodeFeature",
 	"FocusAsset": "/Game/FirstPersonCPP/FirstPersonOverview.FirstPersonOverview",
diff --git a/Templates/FP_FirstPersonBP/manifest.json b/Templates/FP_FirstPersonBP/manifest.json
--- a/Templates/FP_FirstPersonBP/manifest.json
+++ b/Templates/FP_FirstPersonBP/manifest.json
@@ -10,6 +10,10 @@
 			"Language": "ko",
 			"Text": "일인칭"
 		},
+		{
+			"Language": "zh-CN",
+			"Text": "第一人称游戏"
+		},
 		{
 			"Language": "ja",
 			"Text": "ファーストパーソン"
@@ -23,11 +27,15 @@
 		},
 		{
 			"Language": "ko",
-			"Text": "일인칭 팩은 일인칭 시점에서 보이는 (양팔로 표현되는) 플레이어 캐릭터를 특징으로 합니다. 캐릭터 이동은 키보드, 콘트롤러, 터치 디바이스의 가상 조이스틱으로 가능합니다. 추가적으로 플레이어는 마우스, 콘트롤러, 터치 디바이스의 가상 조이스틱을 사용하여 둘러보기도 가능합니다. 캐릭터에는 총이 장착되어 있어, 마우스, 콘트롤러, 터치 디바이스의 가상 조이스틱을 사용하면 레벨과 상호작용하는 '즉시 적중' 무기가 발사됩니다.\n예제 맵도 포함되어 있습니다."
+			"Text": "일인칭 시점에서 보이는 (양팔로 표현되는) 플레이어 캐릭터를 특징으로 합니다. 캐릭터 이동은 키보드, 컨트롤러, 터치 디바이스의 가상 조이스틱으로 가능합니다. 추가적으로 플레이어는 마우스, 컨트롤러, 터치 디바이스의 가상 조이스틱을 사용하여 주변을 둘러볼 수도 있습니다. 캐릭터에는 총이 장착되어 있어 마우스, 컨트롤러, 터치 디바이스의 가상 조이스틱을 사용하면 간단한 발사체가 발사되어, 아레나 벽에 튕기면서 레벨에 있는 몇몇 물리 오브젝트에 영향을 줍니다."
+		},
+		{
+			"Language": "zh-CN",
+			"Text": "第一人称模板拥有一个第一人称的玩家角色（通过一双手臂呈现）。您可以通过键盘、手柄或触摸式设备上的虚拟摇杆控制角色在关卡中四处移动。此外，您还可以使用鼠标、手柄或触摸式设备上的虚拟摇杆转动角色视角，观察四周。角色还装备了枪支，可使用鼠标、手柄或触摸式设备上的虚拟摇杆发射一颗球形子弹，并影响关卡中的一些物理对象，同时从竞技场的墙壁上弹开。"
 		},
 		{
 			"Language": "ja",
-			"Text": "このテンプレートはプレイヤーキャラクター自身の視点で操作するゲームに適しています。キャラクターは武器を持っている腕の部分だけが見えます。このキャラクターはキーボード・コントローラ・バーチャルジョイスティック・タッチパッド等を使ってレベル中を自由に歩き回ることができます。さらに、マウス・コントローラー・バーチャルジョイスティック・タッチパネル等で、視点の向きを変えることができます。装備された武器はマウス・コントローラ・バーチャルジョイスティック・タッチパネルの操作で打つことができます。\nサンプルのマップも用意されています。"
+			"Text": "ファーストパーソンには、一人称視点で表示されるプレイヤーキャラクタ（腕として表示されています）がいます。キーボード、コントローラー、タッチデバイスでのバーチャルジョイスティックでこのキャラクタを移動させることができます。追加して、プレイヤーはマウス、コントローラー、タッチデバイスでのバーチャルジョイスティックで周りを見回すことができます。さらにキャラクタは銃を装備して、マウス、コントローラー、タッチデバイスでのバーチャルジョイスティックを使ってシンプルな球体の弾を発射できます。この弾は物理オブジェクトに影響し、アリーナの壁で反射します。"
 		}
 	],
 	"AssetTypes":
@@ -35,14 +43,6 @@
 		{
 			"Language": "en",
 			"Text": "Animation Sequence, Animation Montage, Sound Wave, Skeletal Mesh, Animation Blueprint, Skeleton, Static Mesh, Texture."
-		},
-		{
-			"Language": "ko",
-			"Text": "애니메이션 시퀀스, 애니메이션 몽타주, 사운드 웨이브, 스켈레탈 메시, 애니메이션 블루프린트, 스켈레톤, 스태틱 메시, 텍스처"
-		},
-		{
-			"Language": "ja",
-			"Text": "アニメーションシーケンス・アニメーションモンタージュ・サウンドウェーブ・スケルタルメッシュ・アニメーションBP・スケルトン・スケルタルメッシュ・テクスチャ"
 		}
 	],
 	"SearchTags":
@@ -50,10 +50,6 @@
 		{
 			"Language": "en",
 			"Text": "First Person,Shoot,Fire,Aim"
-		},
-		{
-			"Language": "ko",
-			"Text": "일인칭,슈팅,발사,조준"
 		}
 	],	
 	"ClassTypes": "Character, GameMode, HUD, Actor, Character, GameMode, HUD, Actor",
diff --git a/Templates/TP_2DSideScroller/Config/TemplateDefs.ini b/Templates/TP_2DSideScroller/Config/TemplateDefs.ini
--- a/Templates/TP_2DSideScroller/Config/TemplateDefs.ini
+++ b/Templates/TP_2DSideScroller/Config/TemplateDefs.ini
@@ -3,7 +3,13 @@ LocalizedDisplayNames=(Language="en",Text="2D Side Scroller")
 LocalizedDescriptions=(Language="en",Text="This is a paper 2D version of the Side Scroller template. It features a 2D sprite character that uses texture based animation where the camera is positioned at the avatar's side. Character movement, which can be controlled with either keyboard, controller or virtual joystick on a touch device, is limited to two dimensions: forward, backward and jumping. A number of ledges are featured in the level and the character can jump between them.")
 
 LocalizedDisplayNames=(Language="ko",Text="2D 횡스크롤")
-LocalizedDescriptions=(Language="ko",Text="횡스크롤 템플릿의 페이퍼 2D 버전입니다. 카메라가 아바타의 측면에 위치하여 텍스처 기반 애니메이션을 사용하는 2D 스프라이트 캐릭터를 특징으로 합니다. 캐릭터 이동은 키보드, 콘트롤러, 터치 디바이스의 가상 조이스틱 중 하나로 가능하며, 이차원으로 제한되어 앞, 뒤, 점프만 가능합니다. 레벨에는 캐릭터가 점프해 다닐 수 있는 턱이 여러개 있습니다.")
+LocalizedDescriptions=(Language="ko",Text="횡스크롤 템플릿의 페이퍼 2D 버전입니다. 카메라가 아바타의 측면에 위치하여 텍스처 기반 애니메이션을 사용하는 2D 스프라이트 캐릭터를 특징으로 합니다. 캐릭터 이동은 키보드, 컨트롤러, 터치 디바이스의 가상 조이스틱 중 하나로 가능하며, 이차원으로 제한되어 앞, 뒤, 점프만 가능합니다. 레벨에는 캐릭터가 점프해 다닐 수 있는 턱이 여러개 있습니다.")
+
+LocalizedDisplayNames=(Language="ja",Text="2Dサイドスクロール")
+LocalizedDescriptions=(Language="ja",Text="サイドスクロールテンプレートの Paper 2D 版です。テクスチャベースのアニメーションを使用する 2D スプライトを使用し、カメラはアバターの側面に存在します。キャラクタの移動はキーボード、コントローラー、タッチデバイスでのバーチャルジョイスティックで操作可能で、二次元に制限されています。前進、後退とジャンプだけです。レベルには足場があり、それらの間をジャンプすることができます。")
+
+LocalizedDisplayNames=(Language="zh-CN",Text="2D横版过关游戏")
+LocalizedDescriptions=(Language="zh-CN",Text="这是一个2D横版过关游戏的模板。它拥有一个2D的Sprite角色，使用基于贴图的动画，摄像机则位于角色的一侧。您可以通过键盘、手柄或触摸式设备上的虚拟摇杆来控制角色移动，但仅限于两个方向：前进、后退及跳跃。关卡中有许多岩架，角色可以在它们之间跳跃。 ")
 
 ClassTypes=PaperCharacter, GameMode, CameraComponent, SpringArmComponent, PaperFlipbook
 AssetTypes=Texture,Sprite, Paper Flipbook, Material
diff --git a/Templates/TP_2DSideScroller/manifest.json b/Templates/TP_2DSideScroller/manifest.json
--- a/Templates/TP_2DSideScroller/manifest.json
+++ b/Templates/TP_2DSideScroller/manifest.json
@@ -13,6 +13,10 @@
 		{
 			"Language": "ja",
 			"Text": "2D横スクロール"
+		},
+		{
+			"Language": "zh-CN",
+			"Text": "2D横版过关游戏"
 		}
 	],
 	"Description":
@@ -23,11 +27,15 @@
 		},
 		{
 			"Language": "ko",
-			"Text": "횡스크롤 팩의 페이퍼 2D 버전입니다. 카메라가 아바타의 측면에 위치하여 텍스처 기반 애니메이션을 사용하는 2D 스프라이트 캐릭터를 특징으로 합니다. 캐릭터 이동은 키보드, 콘트롤러, 터치 디바이스의 가상 조이스틱 중 하나로 가능하며, 이차원으로 제한되어 앞, 뒤, 점프만 가능합니다. 레벨에는 캐릭터가 점프해 다닐 수 있는 턱이 여러개 있습니다."
+			"Text": "횡스크롤 팩의 페이퍼 2D 버전입니다. 카메라가 아바타의 측면에 위치하여 텍스처 기반 애니메이션을 사용하는 2D 스프라이트 캐릭터를 특징으로 합니다. 캐릭터 이동은 키보드, 컨트롤러, 터치 디바이스의 가상 조이스틱 중 하나로 가능하며, 이차원으로 제한되어 앞, 뒤, 점프만 가능합니다. 레벨에는 캐릭터가 점프해 다닐 수 있는 턱이 여러개 있습니다."
 		},
 		{
 			"Language": "ja",
-			"Text": "このテンプレートはペーパー２Dを使った横スクロール系のゲームに適しています。キャラクターは２Dのスプライトとテクスチャ切り替えのアニメーションで設計され、カメラは常にキャラクターを追いかける形になっています。キャラクターはキーボード・コントローラ・タッチパネル上のバーチャルジョイスティックで操作することができます。基本操作は前進・後退・ジャンプの三つです。サンプルのレベルに配置された障害物はキャラクターが飛び越えていける間隔で配置されています。"
+			"Text": "サイドスクロールの Paper 2D 版です。テクスチャベースのアニメーションを使用する 2D スプライトを使用し、カメラはアバターの側面に存在します。キャラクタの移動はキーボード、コントローラー、タッチデバイスでのバーチャルジョイスティックで操作可能で、二次元に制限されています。前進、後退とジャンプだけです。レベルには足場があり、それらの間をジャンプすることができます。"
+		},
+		{
+			"Language": "zh-CN",
+			"Text": "这是一个2D横版过关游戏的模板。它拥有一个2D的Sprite角色，使用基于贴图的动画，摄像机则位于角色的一侧。您可以通过键盘、手柄或触摸式设备上的虚拟摇杆来控制角色移动，但仅限于两个方向：前进、后退及跳跃。关卡中有许多岩架，角色可以在它们之间跳跃。"
 		}
 	],
 	"AssetTypes":
@@ -35,14 +43,6 @@
 		{
 			"Language": "en",
 			"Text": "Texture,Sprite, Sprite Flipbook, Material"
-		},
-		{
-			"Language": "ko",
-			"Text": "텍스처, 스프라이트, 스프라이트 플립북, 머티리얼"
-		},
-		{
-			"Language": "ja",
-			"Text": "テクスチャ・スプライト・スプライトブリップブック・マテリアル"
 		}
 	],	
 	"SearchTags":
@@ -50,10 +50,6 @@
 		{
 			"Language": "en",
 			"Text": "Side,Scroller,Platform,2D,Paper"
-		},
-		{
-			"Language": "ko",
-			"Text": "횡,스크롤,플랫폼,2D,페이퍼"
 		}
 	],
 	"ClassTypes": "PaperCharacter, GameMode, CameraComponent, SpringArmComponent, PaperFlipbook",
diff --git a/Templates/TP_2DSideScrollerBP/Config/TemplateDefs.ini b/Templates/TP_2DSideScrollerBP/Config/TemplateDefs.ini
--- a/Templates/TP_2DSideScrollerBP/Config/TemplateDefs.ini
+++ b/Templates/TP_2DSideScrollerBP/Config/TemplateDefs.ini
@@ -3,7 +3,13 @@ LocalizedDisplayNames=(Language="en",Text="2D Side Scroller")
 LocalizedDescriptions=(Language="en",Text="This is a paper 2D version of the Side Scroller template. It features a 2D sprite character that uses texture based animation where the camera is positioned at the avatar's side. Character movement, which can be controlled with either keyboard, controller or virtual joystick on a touch device, is limited to two dimensions: forward, backward and jumping. A number of ledges are featured in the level and the character can jump between them.")
 
 LocalizedDisplayNames=(Language="ko",Text="2D 횡스크롤")
-LocalizedDescriptions=(Language="ko",Text="횡스크롤 템플릿의 페이퍼 2D 버전입니다. 카메라가 아바타의 측면에 위치하여 텍스처 기반 애니메이션을 사용하는 2D 스프라이트 캐릭터를 특징으로 합니다. 캐릭터 이동은 키보드, 콘트롤러, 터치 디바이스의 가상 조이스틱 중 하나로 가능하며, 이차원으로 제한되어 앞, 뒤, 점프만 가능합니다. 레벨에는 캐릭터가 점프해 다닐 수 있는 턱이 여러개 있습니다.")
+LocalizedDescriptions=(Language="ko",Text="횡스크롤 템플릿의 페이퍼 2D 버전입니다. 카메라가 아바타의 측면에 위치하여 텍스처 기반 애니메이션을 사용하는 2D 스프라이트 캐릭터를 특징으로 합니다. 캐릭터 이동은 키보드, 컨트롤러, 터치 디바이스의 가상 조이스틱 중 하나로 가능하며, 이차원으로 제한되어 앞, 뒤, 점프만 가능합니다. 레벨에는 캐릭터가 점프해 다닐 수 있는 턱이 여러개 있습니다.")
+
+LocalizedDisplayNames=(Language="ja",Text="2Dサイドスクロール")
+LocalizedDescriptions=(Language="ja",Text="サイドスクロールテンプレートの Paper 2D 版です。テクスチャベースのアニメーションを使用する 2D スプライトを使用し、カメラはアバターの側面に存在します。キャラクタの移動はキーボード、コントローラー、タッチデバイスでのバーチャルジョイスティックで操作可能で、二次元に制限されています。前進、後退とジャンプだけです。レベルには足場があり、それらの間をジャンプすることができます。")
+
+LocalizedDisplayNames=(Language="zh-CN",Text="2D横版过关游戏")
+LocalizedDescriptions=(Language="zh-CN",Text="这是一个2D横版过关游戏的模板。它拥有一个2D的Sprite角色，使用基于贴图的动画，摄像机则位于角色的一侧。您可以通过键盘、手柄或触摸式设备上的虚拟摇杆来控制角色移动，但仅限于两个方向：前进、后退及跳跃。关卡中有许多岩架，角色可以在它们之间跳跃。 ")
 
 ClassTypes=APaperCharacter, AGameMode, UCameraComponent, USpringArmComponent, UPaperFlipbook
 AssetTypes=Texture,Sprite, Sprite Flipbook, Material
diff --git a/Templates/TP_2DSideScrollerBP/manifest.json b/Templates/TP_2DSideScrollerBP/manifest.json
--- a/Templates/TP_2DSideScrollerBP/manifest.json
+++ b/Templates/TP_2DSideScrollerBP/manifest.json
@@ -13,6 +13,10 @@
 		{
 			"Language": "ja",
 			"Text": "2D横スクロール"
+		},
+		{
+			"Language": "zh-CN",
+			"Text": "2D横版过关游戏"
 		}
 	],
 	"Description":
@@ -23,11 +27,15 @@
 		},
 		{
 			"Language": "ko",
-			"Text": "횡스크롤 팩의 페이퍼 2D 버전입니다. 카메라가 아바타의 측면에 위치하여 텍스처 기반 애니메이션을 사용하는 2D 스프라이트 캐릭터를 특징으로 합니다. 캐릭터 이동은 키보드, 콘트롤러, 터치 디바이스의 가상 조이스틱 중 하나로 가능하며, 이차원으로 제한되어 앞, 뒤, 점프만 가능합니다. 레벨에는 캐릭터가 점프해 다닐 수 있는 턱이 여러개 있습니다."
+			"Text": "횡스크롤 팩의 페이퍼 2D 버전입니다. 카메라가 아바타의 측면에 위치하여 텍스처 기반 애니메이션을 사용하는 2D 스프라이트 캐릭터를 특징으로 합니다. 캐릭터 이동은 키보드, 컨트롤러, 터치 디바이스의 가상 조이스틱 중 하나로 가능하며, 이차원으로 제한되어 앞, 뒤, 점프만 가능합니다. 레벨에는 캐릭터가 점프해 다닐 수 있는 턱이 여러개 있습니다."
 		},
 		{
 			"Language": "ja",
-			"Text": "このテンプレートはペーパー２Dを使った横スクロール系のゲームに適しています。キャラクターは２Dのスプライトとテクスチャ切り替えのアニメーションで設計され、カメラは常にキャラクターを追いかける形になっています。キャラクターはキーボード・コントローラ・タッチパネル上のバーチャルジョイスティックで操作することができます。基本操作は前進・後退・ジャンプの三つです。サンプルのレベルに配置された障害物はキャラクターが飛び越えていける間隔で配置されています。"
+			"Text": "サイドスクロールの Paper 2D 版です。テクスチャベースのアニメーションを使用する 2D スプライトを使用し、カメラはアバターの側面に存在します。キャラクタの移動はキーボード、コントローラー、タッチデバイスでのバーチャルジョイスティックで操作可能で、二次元に制限されています。前進、後退とジャンプだけです。レベルには足場があり、それらの間をジャンプすることができます。"
+		},
+		{
+			"Language": "zh-CN",
+			"Text": "这是一个2D横版过关游戏的模板。它拥有一个2D的Sprite角色，使用基于贴图的动画，摄像机则位于角色的一侧。您可以通过键盘、手柄或触摸式设备上的虚拟摇杆来控制角色移动，但仅限于两个方向：前进、后退及跳跃。关卡中有许多岩架，角色可以在它们之间跳跃。"
 		}
 	],
 	"AssetTypes":
@@ -35,26 +43,13 @@
 		{
 			"Language": "en",
 			"Text": "Texture,Sprite, Sprite Flipbook, Material"
-		},
-		{
-			"Language": "ko",
-			"Text": "텍스처, 스프라이트, 스프라이트 플립북, 머티리얼"
-		},
-		{
-			"Language": "ja",
-			"Text": "テクスチャ・スプライト・スプライトフリップブック・マテリアル"
 		}
-
 	],	
 	"SearchTags":
 	[
 		{
 			"Language": "en",
 			"Text": "Side,Scroller,Platform,2D,Paper"
-		},
-		{
-			"Language": "ko",
-			"Text": "횡,스크롤,플랫폼,2D,페이퍼"
 		}
 	],
 	"ClassTypes": "PaperCharacter, GameMode, CameraComponent, SpringArmComponent, PaperFlipbook",
diff --git a/Templates/TP_FirstPerson/Config/TemplateDefs.ini b/Templates/TP_FirstPerson/Config/TemplateDefs.ini
--- a/Templates/TP_FirstPerson/Config/TemplateDefs.ini
+++ b/Templates/TP_FirstPerson/Config/TemplateDefs.ini
@@ -3,7 +3,13 @@ LocalizedDisplayNames=(Language="en",Text="First Person")
 LocalizedDescriptions=(Language="en",Text="The first person template features a player character (represented by a pair of arms) which is viewed from first person perspective. The character can be moved around the level using keyboard, controller or virtual joystick on a touch device. Additionally the player can look around using mouse, controller or virtual joystick on a touch device. The character is also equipped with a gun that, using mouse,  controller or virtual joystick on a touch device will fire a simple sphere projectile that will affect some physics objects in the level, whilst bounce off the arena walls.")
 
 LocalizedDisplayNames=(Language="ko",Text="일인칭")
-LocalizedDescriptions=(Language="ko",Text="일인칭 템플릿은 일인칭 시점에서 보이는 (양팔로 표현되는) 플레이어 캐릭터를 특징으로 합니다. 캐릭터 이동은 키보드, 콘트롤러, 터치 디바이스의 가상 조이스틱으로 가능합니다. 추가적으로 플레이어는 마우스, 콘트롤러, 터치 디바이스의 가상 조이스틱을 사용하여 둘러보기도 가능합니다. 캐릭터에는 총이 장착되어 있어, 마우스, 콘트롤러, 터치 디바이스의 가상 조이스틱을 사용하면 간단한 구체 프로젝타일이 발사되어, 아레나 벽에 튕기면서 레벨에 있는 몇몇 물리 오브젝트에 영향을 끼칩니다.")
+LocalizedDescriptions=(Language="ko",Text="일인칭 템플릿은 일인칭 시점에서 보이는 (양팔로 표현되는) 플레이어 캐릭터를 특징으로 합니다. 캐릭터 이동은 키보드, 컨트롤러, 터치 디바이스의 가상 조이스틱으로 가능합니다. 추가적으로 플레이어는 마우스, 컨트롤러, 터치 디바이스의 가상 조이스틱을 사용하여 주변을 둘러볼 수도 있습니다. 캐릭터에는 총이 장착되어 있어 마우스, 컨트롤러, 터치 디바이스의 가상 조이스틱을 사용하면 간단한 발사체가 발사되어, 아레나 벽에 튕기면서 레벨에 있는 몇몇 물리 오브젝트에 영향을 줍니다.")
+
+LocalizedDisplayNames=(Language="ja",Text="ファーストパーソン")
+LocalizedDescriptions=(Language="ja",Text="ファーストパーソンテンプレートには、一人称視点で表示されるプレイヤーキャラクタ（腕として表示されています）がいます。キーボード、コントローラー、タッチデバイスでのバーチャルジョイスティックでこのキャラクタを移動させることができます。追加して、プレイヤーはマウス、コントローラー、タッチデバイスでのバーチャルジョイスティックで周りを見回すことができます。さらにキャラクタは銃を装備して、マウス、コントローラー、タッチデバイスでのバーチャルジョイスティックを使ってシンプルな球体の弾を発射できます。この弾は物理オブジェクトに影響し、アリーナの壁で反射します。")
+
+LocalizedDisplayNames=(Language="zh-CN",Text="第一人称游戏")
+LocalizedDescriptions=(Language="zh-CN",Text="第一人称模板拥有一个第一人称的玩家角色（通过一双手臂呈现）。您可以通过键盘、手柄或触摸式设备上的虚拟摇杆控制角色在关卡中四处移动。此外，您还可以使用鼠标、手柄或触摸式设备上的虚拟摇杆转动角色视角，观察四周。角色还装备了枪支，可使用鼠标、手柄或触摸式设备上的虚拟摇杆发射一颗球形子弹，并影响关卡中的一些物理对象，同时从竞技场的墙壁上弹开。")
 
 ClassTypes=Character, GameMode, HUD, Actor
 AssetTypes=Animation Sequence, Animation Montage, Sound Wave, Skeletal Mesh, Animation Blueprint, Skeleton, Static Mesh, Texture.
diff --git a/Templates/TP_FirstPersonBP/Config/TemplateDefs.ini b/Templates/TP_FirstPersonBP/Config/TemplateDefs.ini
--- a/Templates/TP_FirstPersonBP/Config/TemplateDefs.ini
+++ b/Templates/TP_FirstPersonBP/Config/TemplateDefs.ini
@@ -3,7 +3,13 @@ LocalizedDisplayNames=(Language="en",Text="First Person")
 LocalizedDescriptions=(Language="en",Text="The first person template features a player character (represented by a pair of arms) which is viewed from first person perspective. The character can be moved around the level using keyboard, controller or virtual joystick on a touch device. Additionally the player can look around using mouse, controller or virtual joystick on a touch device. The character is also equipped with a gun that, using mouse,  controller or virtual joystick on a touch device will fire a simple sphere projectile that will affect some physics objects in the level, whilst bounce off the arena walls.")
 
 LocalizedDisplayNames=(Language="ko",Text="일인칭")
-LocalizedDescriptions=(Language="ko",Text="일인칭 템플릿은 일인칭 시점에서 보이는 (양팔로 표현되는) 플레이어 캐릭터를 특징으로 합니다. 캐릭터 이동은 키보드, 콘트롤러, 터치 디바이스의 가상 조이스틱으로 가능합니다. 추가적으로 플레이어는 마우스, 콘트롤러, 터치 디바이스의 가상 조이스틱을 사용하여 둘러보기도 가능합니다. 캐릭터에는 총이 장착되어 있어, 마우스, 콘트롤러, 터치 디바이스의 가상 조이스틱을 사용하면 간단한 구체 프로젝타일이 발사되어, 아레나 벽에 튕기면서 레벨에 있는 몇몇 물리 오브젝트에 영향을 끼칩니다.")
+LocalizedDescriptions=(Language="ko",Text="일인칭 템플릿은 일인칭 시점에서 보이는 (양팔로 표현되는) 플레이어 캐릭터를 특징으로 합니다. 캐릭터 이동은 키보드, 컨트롤러, 터치 디바이스의 가상 조이스틱으로 가능합니다. 추가적으로 플레이어는 마우스, 컨트롤러, 터치 디바이스의 가상 조이스틱을 사용하여 주변을 둘러볼 수도 있습니다. 캐릭터에는 총이 장착되어 있어 마우스, 컨트롤러, 터치 디바이스의 가상 조이스틱을 사용하면 간단한 발사체가 발사되어, 아레나 벽에 튕기면서 레벨에 있는 몇몇 물리 오브젝트에 영향을 줍니다.")
+
+LocalizedDisplayNames=(Language="ja",Text="ファーストパーソン")
+LocalizedDescriptions=(Language="ja",Text="ファーストパーソンテンプレートには、一人称視点で表示されるプレイヤーキャラクタ（腕として表示されています）がいます。キーボード、コントローラー、タッチデバイスでのバーチャルジョイスティックでこのキャラクタを移動させることができます。追加して、プレイヤーはマウス、コントローラー、タッチデバイスでのバーチャルジョイスティックで周りを見回すことができます。さらにキャラクタは銃を装備して、マウス、コントローラー、タッチデバイスでのバーチャルジョイスティックを使ってシンプルな球体の弾を発射できます。この弾は物理オブジェクトに影響し、アリーナの壁で反射します。")
+
+LocalizedDisplayNames=(Language="zh-CN",Text="第一人称游戏")
+LocalizedDescriptions=(Language="zh-CN",Text="第一人称模板拥有一个第一人称的玩家角色（通过一双手臂呈现）。您可以通过键盘、手柄或触摸式设备上的虚拟摇杆控制角色在关卡中四处移动。此外，您还可以使用鼠标、手柄或触摸式设备上的虚拟摇杆转动角色视角，观察四周。角色还装备了枪支，可使用鼠标、手柄或触摸式设备上的虚拟摇杆发射一颗球形子弹，并影响关卡中的一些物理对象，同时从竞技场的墙壁上弹开。")
 
 ClassTypes=Character, GameMode, HUD, Actor
 AssetTypes=Animation Sequence, Animation Montage, Sound Wave, Skeletal Mesh, Animation Blueprint, Skeleton, Static Mesh, Texture.
diff --git a/Templates/TP_Flying/Config/TemplateDefs.ini b/Templates/TP_Flying/Config/TemplateDefs.ini
--- a/Templates/TP_Flying/Config/TemplateDefs.ini
+++ b/Templates/TP_Flying/Config/TemplateDefs.ini
@@ -3,7 +3,13 @@ LocalizedDisplayNames=(Language="en",Text="Flying")
 LocalizedDescriptions=(Language="en",Text="The flying template features a spaceship with the camera positioned behind the ship represented by a static mesh. The keyboard, controller or virtual joystick controls the pitch and yaw of the ship as well as providing a means to allow speed control of the ships forward movement. The level is populated with several clusters of boxes that block the movement of the ship.")
 
 LocalizedDisplayNames=(Language="ko",Text="비행")
-LocalizedDescriptions=(Language="ko",Text="비행 템플릿은 뒤에 카메라가 위치한 스태틱 메시로 표현되는 비행선을 특징으로 합니다. 키보드, 콘트롤러, 가상 조이스틱으로 비행선의 핏치와 요 뿐만 아니라, 비행선의 전방 이동 속도를 제어할 수도 있습니다. 레벨에는 비행선의 이동을 막는 박스 덩어리가 여럿 채워져 있습니다.")
+LocalizedDescriptions=(Language="ko",Text="비행 템플릿은 뒤에 카메라가 위치한 스태틱 메시로 표현되는 비행선을 특징으로 합니다. 키보드, 컨트롤러, 가상 조이스틱으로 비행선의 피치와 요 뿐만 아니라, 비행선의 전방 이동 속도를 제어할 수도 있습니다. 레벨에는 비행선의 이동을 막는 박스 덩어리가 여럿 있습니다.")
+
+LocalizedDisplayNames=(Language="ja",Text="フライング")
+LocalizedDescriptions=(Language="ja",Text="フライングテンプレートには、スタティックメッシュで表現された宇宙船の背後にカメラがあります。キーボード、コントローラー、バーチャルジョイスティックで宇宙船のピッチとヨーのコントロールと、前進時の速度の調整ができます。レベルには積まれた箱があり、宇宙船の移動を邪魔しています。")
+
+LocalizedDisplayNames=(Language="zh-CN",Text="飞行类游戏")
+LocalizedDescriptions=(Language="zh-CN",Text="飞行类游戏模板里拥有一艘宇宙飞船，摄像机位于飞船后方，由静态网格体呈现。键盘、手柄或触摸式设备上的虚拟摇杆可控制飞船的纵倾和航向，并提供了控制飞船前进速度的方法。关卡中还有一些立方体集群，可阻挡飞船的移动。")
 
 ClassTypes=GameMode, Pawn, StaticMeshComponent, SpringArmComponent, CameraComponent, InputComponent, PrimitiveComponent
 AssetTypes=Static Mesh
diff --git a/Templates/TP_Flying/manifest.json b/Templates/TP_Flying/manifest.json
--- a/Templates/TP_Flying/manifest.json
+++ b/Templates/TP_Flying/manifest.json
@@ -10,9 +10,13 @@
 			"Language": "ko",
 			"Text": "비행"
 		},
+		{
+			"Language": "zh-CN",
+			"Text": "飞行类游戏"
+		},
 		{
 			"Language": "ja",
-			"Text": "飛行"
+			"Text": "フライング"
 		}
 	],
 	"Description":
@@ -23,11 +27,15 @@
 		},
 		{
 			"Language": "ko",
-			"Text": "비행 팩은 뒤에 카메라가 위치한 스태틱 메시로 표현되는 비행선을 특징으로 합니다. 키보드, 콘트롤러, 가상 조이스틱으로 비행선의 핏치와 요 뿐만 아니라, 비행선의 전방 이동 속도를 제어할 수도 있습니다. 레벨에는 비행선의 이동을 막는 박스 덩어리가 여럿 채워져 있습니다."
+			"Text": "비행 템플릿은 뒤에 카메라가 위치한 스태틱 메시로 표현되는 비행선을 특징으로 합니다. 키보드, 컨트롤러, 가상 조이스틱으로 비행선의 피치와 요 뿐만 아니라, 비행선의 전방 이동 속도를 제어할 수도 있습니다. 레벨에는 비행선의 이동을 막는 박스 덩어리가 여럿 있습니다."
+		},
+		{
+			"Language": "zh-CN",
+			"Text": "飞行类游戏模板里拥有一艘宇宙飞船，摄像机位于飞船后方，由静态网格体呈现。键盘、手柄或触摸式设备上的虚拟摇杆可控制飞船的纵倾和航向，并提供了控制飞船前进速度的方法。关卡中还有一些立方体集群，可阻挡飞船的移动。"
 		},
 		{
 			"Language": "ja",
-			"Text": "The flying pack features a spaceship with the camera positioned behind the ship represented by a static mesh. The keyboard, controller or virtual joystick controls the pitch and yaw of the ship as well as providing a means to allow speed control of the ships forward movement. The level is populated with several clusters of boxes that block the movement of the ship."
+			"Text": "フライングテンプレートには、スタティックメッシュで表現された宇宙船の背後にカメラがあります。キーボード、コントローラー、バーチャルジョイスティックで宇宙船のピッチとヨーのコントロールと、前進時の速度の調整ができます。レベルには積まれた箱があり、宇宙船の移動を邪魔しています。"
 		}
 
 	],
@@ -36,14 +44,6 @@
 		{
 			"Language": "en",
 			"Text": "Static Mesh"
-		},
-		{
-			"Language": "ko",
-			"Text": "스태틱 메시"
-		},
-		{
-			"Language": "ja",
-			"Text": "スタティックメッシュ"
 		}
 
 	],	
@@ -52,10 +52,6 @@
 		{
 			"Language": "en",
 			"Text": "Flying,Hover,UFO"
-		},
-		{
-			"Language": "ko",
-			"Text": "비행,호버,UFO"
 		}
 	],	
 	"ClassTypes": "GameMode, Pawn, StaticMeshComponent, SpringArmComponent, CameraComponent, InputComponent, PrimitiveComponent",	
diff --git a/Templates/TP_FlyingBP/Config/TemplateDefs.ini b/Templates/TP_FlyingBP/Config/TemplateDefs.ini
--- a/Templates/TP_FlyingBP/Config/TemplateDefs.ini
+++ b/Templates/TP_FlyingBP/Config/TemplateDefs.ini
@@ -3,7 +3,13 @@ LocalizedDisplayNames=(Language="en",Text="Flying")
 LocalizedDescriptions=(Language="en",Text="The flying template features a spaceship with the camera positioned behind the ship represented by a static mesh. The keyboard, controller or virtual joystick controls the pitch and yaw of the ship as well as providing a means to allow speed control of the ships forward movement. The level is populated with several clusters of boxes that block the movement of the ship.")
 
 LocalizedDisplayNames=(Language="ko",Text="비행")
-LocalizedDescriptions=(Language="ko",Text="비행 템플릿은 뒤에 카메라가 위치한 스태틱 메시로 표현되는 비행선을 특징으로 합니다. 키보드, 콘트롤러, 가상 조이스틱으로 비행선의 핏치와 요 뿐만 아니라, 비행선의 전방 이동 속도를 제어할 수도 있습니다. 레벨에는 비행선의 이동을 막는 박스 덩어리가 여럿 채워져 있습니다.")
+LocalizedDescriptions=(Language="ko",Text="비행 템플릿은 뒤에 카메라가 위치한 스태틱 메시로 표현되는 비행선을 특징으로 합니다. 키보드, 컨트롤러, 가상 조이스틱으로 비행선의 피치와 요 뿐만 아니라, 비행선의 전방 이동 속도를 제어할 수도 있습니다. 레벨에는 비행선의 이동을 막는 박스 덩어리가 여럿 있습니다.")
+
+LocalizedDisplayNames=(Language="ja",Text="フライング")
+LocalizedDescriptions=(Language="ja",Text="フライングテンプレートには、スタティックメッシュで表現された宇宙船の背後にカメラがあります。キーボード、コントローラー、バーチャルジョイスティックで宇宙船のピッチとヨーのコントロールと、前進時の速度の調整ができます。レベルには積まれた箱があり、宇宙船の移動を邪魔しています。")
+
+LocalizedDisplayNames=(Language="zh-CN",Text="飞行类游戏")
+LocalizedDescriptions=(Language="zh-CN",Text="飞行类游戏模板里拥有一艘宇宙飞船，摄像机位于飞船后方，由静态网格体呈现。键盘、手柄或触摸式设备上的虚拟摇杆可控制飞船的纵倾和航向，并提供了控制飞船前进速度的方法。关卡中还有一些立方体集群，可阻挡飞船的移动。")
 
 ClassTypes=GameMode, Pawn, StaticMeshComponent, SpringArmComponent, CameraComponent, InputComponent, PrimitiveComponent
 AssetTypes=Static Mesh
diff --git a/Templates/TP_FlyingBP/manifest.json b/Templates/TP_FlyingBP/manifest.json
--- a/Templates/TP_FlyingBP/manifest.json
+++ b/Templates/TP_FlyingBP/manifest.json
@@ -10,9 +10,13 @@
 			"Language": "ko",
 			"Text": "비행"
 		},
+		{
+			"Language": "zh-CN",
+			"Text": "飞行类游戏"
+		},
 		{
 			"Language": "ja",
-			"Text": "飛行"
+			"Text": "フライング"
 		}
 	],
 	"Description":
@@ -23,37 +27,31 @@
 		},
 		{
 			"Language": "ko",
-			"Text": "비행 팩은 뒤에 카메라가 위치한 스태틱 메시로 표현되는 비행선을 특징으로 합니다. 키보드, 콘트롤러, 가상 조이스틱으로 비행선의 핏치와 요 뿐만 아니라, 비행선의 전방 이동 속도를 제어할 수도 있습니다. 레벨에는 비행선의 이동을 막는 박스 덩어리가 여럿 채워져 있습니다."
+			"Text": "비행 템플릿은 뒤에 카메라가 위치한 스태틱 메시로 표현되는 비행선을 특징으로 합니다. 키보드, 컨트롤러, 가상 조이스틱으로 비행선의 피치와 요 뿐만 아니라, 비행선의 전방 이동 속도를 제어할 수도 있습니다. 레벨에는 비행선의 이동을 막는 박스 덩어리가 여럿 있습니다."
+		},
+		{
+			"Language": "zh-CN",
+			"Text": "飞行类游戏模板里拥有一艘宇宙飞船，摄像机位于飞船后方，由静态网格体呈现。键盘、手柄或触摸式设备上的虚拟摇杆可控制飞船的纵倾和航向，并提供了控制飞船前进速度的方法。关卡中还有一些立方体集群，可阻挡飞船的移动。"
 		},
 		{
 			"Language": "ja",
-			"Text": "The flying pack features a spaceship with the camera positioned behind the ship represented by a static mesh. The keyboard, controller or virtual joystick controls the pitch and yaw of the ship as well as providing a means to allow speed control of the ships forward movement. The level is populated with several clusters of boxes that block the movement of the ship."
+			"Text": "フライングテンプレートには、スタティックメッシュで表現された宇宙船の背後にカメラがあります。キーボード、コントローラー、バーチャルジョイスティックで宇宙船のピッチとヨーのコントロールと、前進時の速度の調整ができます。レベルには積まれた箱があり、宇宙船の移動を邪魔しています。"
 		}
+
 	],
 	"AssetTypes":
 	[
 		{
 			"Language": "en",
 			"Text": "Static Mesh"
-		},
-		{
-			"Language": "ko",
-			"Text": "스태틱 메시"
-		},
-		{
-			"Language": "ja",
-			"Text": "スタティックメッシュ"
 		}
+
 	],	
 	"SearchTags":
 	[
 		{
 			"Language": "en",
 			"Text": "Flying,Hover,UFO"
-		},
-		{
-			"Language": "ko",
-			"Text": "비행,호버,UFO"
 		}
 	],	
 	"ClassTypes": "GameMode, Pawn, StaticMeshComponent, SpringArmComponent, CameraComponent, InputComponent, PrimitiveComponent",
diff --git a/Templates/TP_HandheldARBP/Config/TemplateDefs.ini b/Templates/TP_HandheldARBP/Config/TemplateDefs.ini
--- a/Templates/TP_HandheldARBP/Config/TemplateDefs.ini
+++ b/Templates/TP_HandheldARBP/Config/TemplateDefs.ini
@@ -3,7 +3,13 @@ LocalizedDisplayNames=(Language="en",Text="Handheld AR")
 LocalizedDescriptions=(Language="en",Text="A starting point for building Augmented Reality applications for Android and iOS devices. Includes runtime logic for toggling AR mode on and off, debug information regarding plane detection, example code for hit detection and handling light estimation.")
 
 LocalizedDisplayNames=(Language="ko",Text="휴대형 증강 현실")
-LocalizedDescriptions=(Language="ko",Text="안드로이드 및 iOS 디바이스용 증강 현실 어플리케이션을 만들기 위한 템플릿입니다. AR 모드를 켜고 끄는 런타임 로직, 면 인식 관련 디버깅 정보, 히트 디텍션과 빛 추정 처리 관련 예제 코드가 포함되어 있습니다.")
+LocalizedDescriptions=(Language="ko",Text="Android 및 iOS 디바이스용 증강 현실 앱을 만들기 위한 템플릿입니다. AR 모드를 켜고 끄는 런타임 로직, 면 인식 관련 디버깅 정보, 히트 디텍션과 빛 추정 처리 관련 예제 코드가 포함되어 있습니다.")
+
+LocalizedDisplayNames=(Language="ja",Text="携帯 AR")
+LocalizedDescriptions=(Language="ja",Text="Android と iOS 向けの拡張現実アプリケーション作成の土台です。AR モードのオンオフを行う実行時のロジック、平面検出に関するデバッグ情報、ヒット検出のサンプルコード、ライト推定などを含んでいます。")
+
+LocalizedDisplayNames=(Language="zh-CN",Text="手持式AR应用")
+LocalizedDescriptions=(Language="zh-CN",Text="为安卓及iOS设备创建增强现实应用的良好起点，包含开启和关闭AR模式的运行逻辑，关于平面检测的调试信息，命中检测及处理预计光照的示例代码。")
 
 ClassTypes=Pawn, GameMode, Widget, Actor
 AssetTypes=Static Mesh
diff --git a/Templates/TP_HandheldARBP/manifest.json b/Templates/TP_HandheldARBP/manifest.json
--- a/Templates/TP_HandheldARBP/manifest.json
+++ b/Templates/TP_HandheldARBP/manifest.json
@@ -5,6 +5,18 @@
 		{
 			"Language": "en",
 			"Text": "Handheld AR"
+		},
+		{
+			"Language": "ko",
+			"Text": "휴대형 증강 현실"
+		},
+		{
+			"Language": "ja",
+			"Text": "携帯 AR"
+		},
+		{
+			"Language": "zh-CN",
+			"Text": "手持式AR应用"
 		}
 	],
 	"Description":
@@ -12,6 +24,18 @@
     {
       "Language": "en",
       "Text": "A starting point for building Augmented Reality applications for Android and iOS devices. Includes runtime logic for toggling AR mode on and off, debug information regarding plane detection, example code for hit detection and handling light estimation."
+    },
+    {
+      "Language": "ko",
+      "Text": "Android 및 iOS 디바이스용 증강 현실 앱을 만들기 위한 템플릿입니다. AR 모드를 켜고 끄는 런타임 로직, 면 인식 관련 디버깅 정보, 히트 디텍션과 빛 추정 처리 관련 예제 코드가 포함되어 있습니다."
+    },
+    {
+      "Language": "ja",
+      "Text": "Android と iOS 向けの拡張現実アプリケーション作成の土台です。AR モードのオンオフを行う実行時のロジック、平面検出に関するデバッグ情報、ヒット検出のサンプルコード、ライト推定などを含んでいます。"
+    },
+    {
+      "Language": "zh-CN",
+      "Text": "为安卓及iOS设备创建增强现实应用的良好起点，包含开启和关闭AR模式的运行逻辑，关于平面检测的调试信息，命中检测及处理预计光照的示例代码。"
     }
 	],
 	"AssetTypes":
diff --git a/Templates/TP_Puzzle/Config/TemplateDefs.ini b/Templates/TP_Puzzle/Config/TemplateDefs.ini
--- a/Templates/TP_Puzzle/Config/TemplateDefs.ini
+++ b/Templates/TP_Puzzle/Config/TemplateDefs.ini
@@ -3,7 +3,13 @@ LocalizedDisplayNames=(Language="en",Text="Puzzle")
 LocalizedDescriptions=(Language="en",Text="The puzzle template features a simple grid of blocks. The camera is postioned overhead of the blocks. The user can interact with these blocks by clicking using either the mouse or touch screen. When a block is clicked its color changes and a score, which shown using a text component within the level is increased.")
 
 LocalizedDisplayNames=(Language="ko",Text="퍼즐")
-LocalizedDescriptions=(Language="ko",Text="퍼즐 템플릿은 단순한 블록 모눈을 특징으로 합니다. 카메라는 블록 위에 위치해 있습니다. 사용자는 이 블록을 마우스나 터치스크린으로 클릭하여 상호작용할 수 있습니다. 블록을 클릭하면 색이 바뀌며, 레벨에 텍스트 컴포넌트를 사용하여 표시되는 점수가 올라갑니다.")
+LocalizedDescriptions=(Language="ko",Text="퍼즐 템플릿은 단순한 블록 모눈을 특징으로 합니다. 카메라는 블록 위에 위치해 있습니다. 사용자는 이 블록을 마우스나 터치스크린으로 클릭하여 상호작용할 수 있습니다. 블록을 클릭하면 색이 바뀌면서 레벨에 표시된 점수가 올라가며, 여기에는 텍스트 컴포넌트가 사용되었습니다.")
+
+LocalizedDisplayNames=(Language="ja",Text="パズル")
+LocalizedDescriptions=(Language="ja",Text="パズルテンプレートにはシンプルなブロックのグリッドが含まれています。カメラはブロックを上から見下ろしています。ユーザーはマウスやタッチスクリーンでブロックをクリックすることができます。ブロックをクリックすると色が変化し得点が追加されます。得点はテキストコンポーネントで表示されています。")
+
+LocalizedDisplayNames=(Language="zh-CN",Text="拼图类游戏")
+LocalizedDescriptions=(Language="zh-CN",Text="拼图类游戏模板拥有一张简单的方块网格。摄像机位于方块的正上方。用户可以通过鼠标或触摸屏点击来与这些方块交互。被点击的方块会改变颜色，同时关卡内使用文本组件显示的分数会增加。")
 
 ClassTypes=Actor, PlayerController, SceneComponent, TextRenderComponent, StaticMeshComponent, MaterialInstance, PrimitiveComponent
 AssetTypes=Static Mesh 
diff --git a/Templates/TP_Puzzle/manifest.json b/Templates/TP_Puzzle/manifest.json
--- a/Templates/TP_Puzzle/manifest.json
+++ b/Templates/TP_Puzzle/manifest.json
@@ -10,6 +10,10 @@
 			"Language": "ko",
 			"Text": "퍼즐"
 		},
+		{
+			"Language": "zh-CN",
+			"Text": "拼图类游戏"
+		},
 		{
 			"Language": "ja",
 			"Text": "パズル"
@@ -23,11 +27,15 @@
 		},
 		{
 			"Language": "ko",
-			"Text": "퍼즐 팩은 단순한 블록 모눈을 특징으로 합니다. 카메라는 블록 위에 위치해 있습니다. 사용자는 이 블록을 마우스나 터치스크린으로 클릭하여 상호작용할 수 있습니다. 블록을 클릭하면 색이 바뀌며, 레벨에 텍스트 컴포넌트를 사용하여 표시되는 점수가 올라갑니다."
+			"Text": "퍼즐 템플릿은 단순한 블록 모눈을 특징으로 합니다. 카메라는 블록 위에 위치해 있습니다. 사용자는 이 블록을 마우스나 터치스크린으로 클릭하여 상호작용할 수 있습니다. 블록을 클릭하면 색이 바뀌면서 레벨에 표시된 점수가 올라가며, 여기에는 텍스트 컴포넌트가 사용되었습니다."
+		},
+		{
+			"Language": "zh-CN",
+			"Text": "拼图类游戏模板拥有一张简单的方块网格。摄像机位于方块的正上方。用户可以通过鼠标或触摸屏点击来与这些方块交互。被点击的方块会改变颜色，同时关卡内使用文本组件显示的分数会增加。"
 		},
 		{
 			"Language": "ja",
-			"Text": "The puzzle pack features a simple grid of blocks. The camera is positioned overhead of the blocks. The user can interact with these blocks by clicking using either the mouse or touch screen. When a block is clicked its color changes and a score, which shown using a text component within the level is increased."
+			"Text": "パズルテンプレートにはシンプルなブロックのグリッドが含まれています。カメラはブロックを上から見下ろしています。ユーザーはマウスやタッチスクリーンでブロックをクリックすることができます。ブロックをクリックすると色が変化し得点が追加されます。得点はテキストコンポーネントで表示されています。"
 		}
 	],
 	"AssetTypes":
@@ -35,14 +43,6 @@
 		{
 			"Language": "en",
 			"Text": "Static Mesh"
-		},
-		{
-			"Language": "ko",
-			"Text": "스태틱 메시"
-		},
-		{
-			"Language": "ja",
-			"Text": "スタティックメッシュ"
 		}
 	],	
 	"SearchTags":
@@ -50,10 +50,6 @@
 		{
 			"Language": "en",
 			"Text": "Puzzle,Blocks,Match"
-		},
-		{
-			"Language": "ko",
-			"Text": "퍼즐,블록,짝찾기"
 		}
 	],	
 	"ClassTypes": "Actor, PlayerController, SceneComponent, TextRenderComponent, StaticMeshComponent, MaterialInstance, PrimitiveComponent",
diff --git a/Templates/TP_PuzzleBP/Config/TemplateDefs.ini b/Templates/TP_PuzzleBP/Config/TemplateDefs.ini
--- a/Templates/TP_PuzzleBP/Config/TemplateDefs.ini
+++ b/Templates/TP_PuzzleBP/Config/TemplateDefs.ini
@@ -3,7 +3,13 @@ LocalizedDisplayNames=(Language="en",Text="Puzzle")
 LocalizedDescriptions=(Language="en",Text="The puzzle template features a simple grid of blocks. The camera is postioned overhead of the blocks. The user can interact with these blocks by clicking using either the mouse or touch screen. When a block is clicked its color changes and a score, which shown using a text component within the level is increased.")
 
 LocalizedDisplayNames=(Language="ko",Text="퍼즐")
-LocalizedDescriptions=(Language="ko",Text="퍼즐 템플릿은 단순한 블록 모눈을 특징으로 합니다. 카메라는 블록 위에 위치해 있습니다. 사용자는 이 블록을 마우스나 터치스크린으로 클릭하여 상호작용할 수 있습니다. 블록을 클릭하면 색이 바뀌며, 레벨에 텍스트 컴포넌트를 사용하여 표시되는 점수가 올라갑니다.")
+LocalizedDescriptions=(Language="ko",Text="퍼즐 템플릿은 단순한 블록 모눈을 특징으로 합니다. 카메라는 블록 위에 위치해 있습니다. 사용자는 이 블록을 마우스나 터치스크린으로 클릭하여 상호작용할 수 있습니다. 블록을 클릭하면 색이 바뀌면서 레벨에 표시된 점수가 올라가며, 여기에는 텍스트 컴포넌트가 사용되었습니다.")
+
+LocalizedDisplayNames=(Language="ja",Text="パズル")
+LocalizedDescriptions=(Language="ja",Text="パズルテンプレートにはシンプルなブロックのグリッドが含まれています。カメラはブロックを上から見下ろしています。ユーザーはマウスやタッチスクリーンでブロックをクリックすることができます。ブロックをクリックすると色が変化し得点が追加されます。得点はテキストコンポーネントで表示されています。")
+
+LocalizedDisplayNames=(Language="zh-CN",Text="拼图类游戏")
+LocalizedDescriptions=(Language="zh-CN",Text="拼图类游戏模板拥有一张简单的方块网格。摄像机位于方块的正上方。用户可以通过鼠标或触摸屏点击来与这些方块交互。被点击的方块会改变颜色，同时关卡内使用文本组件显示的分数会增加。")
 
 ClassTypes=Actor, PlayerController, SceneComponent, TextRenderComponent, StaticMeshComponent, MaterialInstance, PrimitiveComponent
 AssetTypes=Static Mesh
diff --git a/Templates/TP_PuzzleBP/manifest.json b/Templates/TP_PuzzleBP/manifest.json
--- a/Templates/TP_PuzzleBP/manifest.json
+++ b/Templates/TP_PuzzleBP/manifest.json
@@ -10,6 +10,10 @@
 			"Language": "ko",
 			"Text": "퍼즐"
 		},
+		{
+			"Language": "zh-CN",
+			"Text": "拼图类游戏"
+		},
 		{
 			"Language": "ja",
 			"Text": "パズル"
@@ -23,11 +27,15 @@
 		},
 		{
 			"Language": "ko",
-			"Text": "퍼즐 팩은 단순한 블록 모눈을 특징으로 합니다. 카메라는 블록 위에 위치해 있습니다. 사용자는 이 블록을 마우스나 터치스크린으로 클릭하여 상호작용할 수 있습니다. 블록을 클릭하면 색이 바뀌며, 레벨에 텍스트 컴포넌트를 사용하여 표시되는 점수가 올라갑니다."
+			"Text": "퍼즐 템플릿은 단순한 블록 모눈을 특징으로 합니다. 카메라는 블록 위에 위치해 있습니다. 사용자는 이 블록을 마우스나 터치스크린으로 클릭하여 상호작용할 수 있습니다. 블록을 클릭하면 색이 바뀌면서 레벨에 표시된 점수가 올라가며, 여기에는 텍스트 컴포넌트가 사용되었습니다."
+		},
+		{
+			"Language": "zh-CN",
+			"Text": "拼图类游戏模板拥有一张简单的方块网格。摄像机位于方块的正上方。用户可以通过鼠标或触摸屏点击来与这些方块交互。被点击的方块会改变颜色，同时关卡内使用文本组件显示的分数会增加。"
 		},
 		{
 			"Language": "ja",
-			"Text": "The puzzle pack features a simple grid of blocks. The camera is positioned overhead of the blocks. The user can interact with these blocks by clicking using either the mouse or touch screen. When a block is clicked its color changes and a score, which shown using a text component within the level is increased."
+			"Text": "パズルテンプレートにはシンプルなブロックのグリッドが含まれています。カメラはブロックを上から見下ろしています。ユーザーはマウスやタッチスクリーンでブロックをクリックすることができます。ブロックをクリックすると色が変化し得点が追加されます。得点はテキストコンポーネントで表示されています。"
 		}
 	],
 	"AssetTypes":
@@ -35,14 +43,6 @@
 		{
 			"Language": "en",
 			"Text": "Static Mesh"
-		},
-		{
-			"Language": "ko",
-			"Text": "스태틱 메시"
-		},
-		{
-			"Language": "ja",
-			"Text": "スタティックメッシュ"
 		}
 	],	
 	"SearchTags":
@@ -50,10 +50,6 @@
 		{
 			"Language": "en",
 			"Text": "Puzzle,Blocks,Match"
-		},
-		{
-			"Language": "ko",
-			"Text": "퍼즐,블록,짝찾기"
 		}
 	],	
 	"ClassTypes": "Actor, PlayerController, SceneComponent, TextRenderComponent, StaticMeshComponent, MaterialInstance, PrimitiveComponent",
diff --git a/Templates/TP_Rolling/Config/TemplateDefs.ini b/Templates/TP_Rolling/Config/TemplateDefs.ini
--- a/Templates/TP_Rolling/Config/TemplateDefs.ini
+++ b/Templates/TP_Rolling/Config/TemplateDefs.ini
@@ -3,7 +3,13 @@ LocalizedDisplayNames=(Language="en",Text="Rolling")
 LocalizedDescriptions=(Language="en",Text="The rolling template features a physics-based rolling ball which can be rolled in any direction using either a keyboard, controller or virtual joystick on touch devices. The camera is positioned above and behind the ball and follows the movement of the ball. Within the level are some ramps, several physics cubes in the level that can be knocked over with the ball as well as some cubes which cannot.")
 
 LocalizedDisplayNames=(Language="ko",Text="굴리기")
-LocalizedDescriptions=(Language="ko",Text="굴리기 템플릿은 키보드, 콘트롤러, 터치 디바이스의 가상 조이스틱을 사용하여 어느 방향이든 굴릴 수 있는 물리 기반 공을 특징으로 합니다. 카메라는 공의 상단 뒤쪽에 위치해 있어 공의 이동을 따라갑니다. 레벨 안에는 경사로와, 공으로 넘어뜨릴 수도 넘어뜨리지 못할수도 있는 물리 큐브가 여럿 있습니다.")
+LocalizedDescriptions=(Language="ko",Text="굴리기 템플릿은 키보드, 컨트롤러, 터치 디바이스의 가상 조이스틱을 사용하여 어느 방향이든 굴릴 수 있는 물리 기반 공을 특징으로 합니다. 카메라는 공의 상단 뒤쪽에 위치해 있어 공의 이동을 따라갑니다. 레벨 안에는 경사로와, 공으로 넘어뜨릴 수도 넘어뜨리지 못할수도 있는 물리 큐브가 여럿 있습니다.")
+
+LocalizedDisplayNames=(Language="ja",Text="ローリング")
+LocalizedDescriptions=(Language="ja",Text="ローリングテンプレートには物理シミュレーションを行っているボールが含まれ、キーボード、コントローラー、タッチデバイスでのバーチャルジョイスティックで好きな方向に転がすことができます。カメラはボールの上方の後方に位置し、ボールの移動に追従します。レベル内には坂や、ボールで倒すことのできる物理キューブと、物理ではないキューブが含まれています。")
+
+LocalizedDisplayNames=(Language="zh-CN",Text="滚球类游戏")
+LocalizedDescriptions=(Language="zh-CN",Text="滚球类游戏模板中有一个基于物理的滚球，您可以通过键盘、手柄或触摸式设备的虚拟摇杆让其向任何方向滚动。摄像机位于球体的后上方，并会跟随球体移动。在关卡内还有一些斜坡，以及一些可以用球体撞开的物理立方体和无法撞开的立方体。")
 
 ClassTypes=GameMode, Pawn, StaticMeshComponent, SpringArmComponent, CameraComponent, PrimitiveComponent
 AssetTypes=StaticMesh, Material, Physical Material
diff --git a/Templates/TP_Rolling/manifest.json b/Templates/TP_Rolling/manifest.json
--- a/Templates/TP_Rolling/manifest.json
+++ b/Templates/TP_Rolling/manifest.json
@@ -10,9 +10,13 @@
 			"Language": "ko",
 			"Text": "굴리기"
 		},
+		{
+			"Language": "zh-CN",
+			"Text": "滚球类游戏"
+		},
 		{
 			"Language": "ja",
-			"Text": "転がり"
+			"Text": "ローリング"
 		}
 	],
 	"Description":
@@ -23,11 +27,15 @@
 		},
 		{
 			"Language": "ko",
-			"Text": "굴리기 팩은 키보드, 콘트롤러, 터치 디바이스의 가상 조이스틱을 사용하여 어느 방향이든 굴릴 수 있는 물리 기반 공을 특징으로 합니다. 카메라는 공의 상단 뒤쪽에 위치해 있어 공의 이동을 따라갑니다. 레벨 안에는 경사로와, 공으로 넘어뜨릴 수도 넘어뜨리지 못할수도 있는 물리 큐브가 여럿 있습니다."
+			"Text": "굴리기 템플릿은 키보드, 컨트롤러, 터치 디바이스의 가상 조이스틱을 사용하여 어느 방향이든 굴릴 수 있는 물리 기반 공을 특징으로 합니다. 카메라는 공의 상단 뒤쪽에 위치해 있어 공의 이동을 따라갑니다. 레벨 안에는 경사로와, 공으로 넘어뜨릴 수도 넘어뜨리지 못할수도 있는 물리 큐브가 여럿 있습니다."
+		},
+		{
+			"Language": "zh-CN",
+			"Text": "滚球类游戏模板中有一个基于物理的滚球，您可以通过键盘、手柄或触摸式设备的虚拟摇杆让其向任何方向滚动。摄像机位于球体的后上方，并会跟随球体移动。在关卡内还有一些斜坡，以及一些可以用球体撞开的物理立方体和无法撞开的立方体。"
 		},
 		{
 			"Language": "ja",
-			"Text": "転がりテンプレートでは、物理ベースで転がるボールを操作するゲームを作ることができます。ボールは、キーボード・コントローラー・バーチャルジョイスティック・タッチパネルで操作することができます。カメラはボールの上空少し後ろからボールの動きを捉えています。テンプレートに含まれるレベルには幾つかの坂になった障害物が設置されています。レベル内に配置されている幾つかの箱はボールで押して動かすことができますが、できないものも含まれています。"
+			"Text": "ローリングテンプレートには物理シミュレーションを行っているボールが含まれ、キーボード、コントローラー、タッチデバイスでのバーチャルジョイスティックで好きな方向に転がすことができます。カメラはボールの上方の後方に位置し、ボールの移動に追従します。レベル内には坂や、ボールで倒すことのできる物理キューブと、物理ではないキューブが含まれています。"
 		}
 	],
 	"AssetTypes":
@@ -35,14 +43,6 @@
 		{
 			"Language": "en",
 			"Text": "StaticMesh, Material, Physical Material"
-		},
-		{
-			"Language": "ko",
-			"Text": "스태틱 메시, 머티리얼, 피지컬 머티리얼"
-		},
-		{
-			"Language": "ja",
-			"Text": "スタティックメッシュ・マテリアル・フィジカルマテリアル"
 		}
 	],
 	"SearchTags":
@@ -50,10 +50,6 @@
 		{
 			"Language": "en",
 			"Text": "Rolling,Marble"
-		},
-		{
-			"Language": "ko",
-			"Text": "굴리기,마블"
 		}
 	],	
 	"AdditionalFeaturePacks":
diff --git a/Templates/TP_RollingBP/Config/TemplateDefs.ini b/Templates/TP_RollingBP/Config/TemplateDefs.ini
--- a/Templates/TP_RollingBP/Config/TemplateDefs.ini
+++ b/Templates/TP_RollingBP/Config/TemplateDefs.ini
@@ -3,7 +3,13 @@ LocalizedDisplayNames=(Language="en",Text="Rolling")
 LocalizedDescriptions=(Language="en",Text="The rolling template features a physics-based rolling ball which can be rolled in any direction using either a keyboard, controller or virtual joystick on touch devices. The camera is positioned above and behind the ball and follows the movement of the ball. Within the level are some ramps, several physics cubes in the level that can be knocked over with the ball as well as some cubes which cannot.")
 
 LocalizedDisplayNames=(Language="ko",Text="굴리기")
-LocalizedDescriptions=(Language="ko",Text="굴리기 템플릿은 키보드, 콘트롤러, 터치 디바이스의 가상 조이스틱을 사용하여 어느 방향이든 굴릴 수 있는 물리 기반 공을 특징으로 합니다. 카메라는 공의 상단 뒤쪽에 위치해 있어 공의 이동을 따라갑니다. 레벨 안에는 경사로와, 공으로 넘어뜨릴 수도 넘어뜨리지 못할수도 있는 물리 큐브가 여럿 있습니다.")
+LocalizedDescriptions=(Language="ko",Text="굴리기 템플릿은 키보드, 컨트롤러, 터치 디바이스의 가상 조이스틱을 사용하여 어느 방향이든 굴릴 수 있는 물리 기반 공을 특징으로 합니다. 카메라는 공의 상단 뒤쪽에 위치해 있어 공의 이동을 따라갑니다. 레벨 안에는 경사로와, 공으로 넘어뜨릴 수도 넘어뜨리지 못할수도 있는 물리 큐브가 여럿 있습니다.")
+
+LocalizedDisplayNames=(Language="ja",Text="ローリング")
+LocalizedDescriptions=(Language="ja",Text="ローリングテンプレートには物理シミュレーションを行っているボールが含まれ、キーボード、コントローラー、タッチデバイスでのバーチャルジョイスティックで好きな方向に転がすことができます。カメラはボールの上方の後方に位置し、ボールの移動に追従します。レベル内には坂や、ボールで倒すことのできる物理キューブと、物理ではないキューブが含まれています。")
+
+LocalizedDisplayNames=(Language="zh-CN",Text="滚球类游戏")
+LocalizedDescriptions=(Language="zh-CN",Text="滚球类游戏模板中有一个基于物理的滚球，您可以通过键盘、手柄或触摸式设备的虚拟摇杆让其向任何方向滚动。摄像机位于球体的后上方，并会跟随球体移动。在关卡内还有一些斜坡，以及一些可以用球体撞开的物理立方体和无法撞开的立方体。")
 
 ClassTypes=GameMode, Pawn, StaticMeshComponent, SpringArmComponent, CameraComponent, PrimitiveComponent
 AssetTypes=StaticMesh, Material, Physical Material
diff --git a/Templates/TP_RollingBP/manifest.json b/Templates/TP_RollingBP/manifest.json
--- a/Templates/TP_RollingBP/manifest.json
+++ b/Templates/TP_RollingBP/manifest.json
@@ -10,9 +10,13 @@
 			"Language": "ko",
 			"Text": "굴리기"
 		},
+		{
+			"Language": "zh-CN",
+			"Text": "滚球类游戏"
+		},
 		{
 			"Language": "ja",
-			"Text": "転がり"
+			"Text": "ローリング"
 		}
 	],
 	"Description":
@@ -23,11 +27,15 @@
 		},
 		{
 			"Language": "ko",
-			"Text": "굴리기 팩은 키보드, 콘트롤러, 터치 디바이스의 가상 조이스틱을 사용하여 어느 방향이든 굴릴 수 있는 물리 기반 공을 특징으로 합니다. 카메라는 공의 상단 뒤쪽에 위치해 있어 공의 이동을 따라갑니다. 레벨 안에는 경사로와, 공으로 넘어뜨릴 수도 넘어뜨리지 못할수도 있는 물리 큐브가 여럿 있습니다."
+			"Text": "굴리기 템플릿은 키보드, 컨트롤러, 터치 디바이스의 가상 조이스틱을 사용하여 어느 방향이든 굴릴 수 있는 물리 기반 공을 특징으로 합니다. 카메라는 공의 상단 뒤쪽에 위치해 있어 공의 이동을 따라갑니다. 레벨 안에는 경사로와, 공으로 넘어뜨릴 수도 넘어뜨리지 못할수도 있는 물리 큐브가 여럿 있습니다."
+		},
+		{
+			"Language": "zh-CN",
+			"Text": "滚球类游戏模板中有一个基于物理的滚球，您可以通过键盘、手柄或触摸式设备的虚拟摇杆让其向任何方向滚动。摄像机位于球体的后上方，并会跟随球体移动。在关卡内还有一些斜坡，以及一些可以用球体撞开的物理立方体和无法撞开的立方体。"
 		},
 		{
 			"Language": "ja",
-			"Text": "転がりテンプレートでは、物理ベースで転がるボールを操作するゲームを作ることができます。ボールは、キーボード・コントローラー・バーチャルジョイスティック・タッチパネルで操作することができます。カメラはボールの上空少し後ろからボールの動きを捉えています。テンプレートに含まれるレベルには幾つかの坂になった障害物が設置されています。レベル内に配置されている幾つかの箱はボールで押して動かすことができますが、できないものも含まれています。"
+			"Text": "ローリングテンプレートには物理シミュレーションを行っているボールが含まれ、キーボード、コントローラー、タッチデバイスでのバーチャルジョイスティックで好きな方向に転がすことができます。カメラはボールの上方の後方に位置し、ボールの移動に追従します。レベル内には坂や、ボールで倒すことのできる物理キューブと、物理ではないキューブが含まれています。"
 		}
 	],
 	"AdditionalFeaturePacks":
diff --git a/Templates/TP_SideScroller/Config/TemplateDefs.ini b/Templates/TP_SideScroller/Config/TemplateDefs.ini
--- a/Templates/TP_SideScroller/Config/TemplateDefs.ini
+++ b/Templates/TP_SideScroller/Config/TemplateDefs.ini
@@ -3,7 +3,13 @@ LocalizedDisplayNames=(Language="en",Text="Side Scroller")
 LocalizedDescriptions=(Language="en",Text="This side scroller template features an animated skeletal mesh character with a camera positioned at the character's side. Movement, which can be controlled with either keyboard, controller or virtual joystick on a touch device is limited to two dimensions: forward, backward and jumping. A number of ledges are featured in the level and the character can jump between them.")
 
 LocalizedDisplayNames=(Language="ko",Text="횡스크롤")
-LocalizedDescriptions=(Language="ko",Text="횡스크롤 템플릿에는 측면에 카메라가 위치한 애니메이션 적용 스켈레탈 메시 캐릭터를 특징으로 합니다. 키보드, 콘트롤러, 터치 디바이스의 가상 조이스틱으로 제어할 수 있는 이동은 두 축으로 제한되어 있어, 앞, 뒤, 점프만 가능합니다. 레벨에는 캐릭터가 점프해 다닐 수 있는 턱이 다수 있습니다.")
+LocalizedDescriptions=(Language="ko",Text="횡스크롤 템플릿에는 카메라가 측면에 위치한 애니메이션 적용 스켈레탈 메시 캐릭터를 특징으로 합니다. 키보드, 컨트롤러, 터치 디바이스의 가상 조이스틱으로 제어할 수 있는 이동은 두 축으로 제한되어 있어, 앞, 뒤, 점프만 가능합니다. 레벨에는 캐릭터가 점프해 다닐 수 있는 턱이 다수 있습니다.")
+
+LocalizedDisplayNames=(Language="ja",Text="サイドスクロール")
+LocalizedDescriptions=(Language="ja",Text="サイドスクロールテンプレートにはアニメーションを使用するスケルタルメッシュのキャラクタが存在し、カメラはキャラクタの側面に存在します。移動はキーボード、コントローラー、タッチデバイスでのバーチャルジョイスティックで操作可能で、二次元に制限されています。前進、後退とジャンプだけです。レベルには足場があり、それらの間をジャンプすることができます。")
+
+LocalizedDisplayNames=(Language="zh-CN",Text="横版过关游戏")
+LocalizedDescriptions=(Language="zh-CN",Text="该横版过关类游戏模板拥有一个可活动的骨架网格体角色，摄像机位于角色的一侧。您可以通过键盘、手柄或触摸式设备上的虚拟摇杆来控制角色移动，但仅限于两个方向：前进、后退及跳跃。关卡中有许多岩架，角色可以在它们之间跳跃。")
 
 ClassTypes=CameraComponent, SpringArmComponent, PaperFlipbook, PaperCharacter, GameMode
 AssetTypes=Animation Sequence, Animation Blueprint, Skeletal Mesh, Blend Space 1D, Skeleton
diff --git a/Templates/TP_SideScroller/manifest.json b/Templates/TP_SideScroller/manifest.json
--- a/Templates/TP_SideScroller/manifest.json
+++ b/Templates/TP_SideScroller/manifest.json
@@ -10,6 +10,10 @@
 			"Language": "ko",
 			"Text": "횡스크롤"
 		},
+		{
+			"Language": "zh-CN",
+			"Text": "横版过关游戏"
+		},
 		{
 			"Language": "ja",
 			"Text": "横スクロール"
@@ -23,11 +27,15 @@
 		},
 		{
 			"Language": "ko",
-			"Text": "횡스크롤 팩에는 측면에 카메라가 위치한 애니메이션 적용 스켈레탈 메시 캐릭터를 특징으로 합니다. 키보드, 콘트롤러, 터치 디바이스의 가상 조이스틱으로 제어할 수 있는 이동은 두 축으로 제한되어 있어, 앞, 뒤, 점프만 가능합니다. 레벨에는 캐릭터가 점프해 다닐 수 있는 턱이 다수 있습니다."
+			"Text": "횡스크롤 템플릿에는 카메라가 측면에 위치한 애니메이션 적용 스켈레탈 메시 캐릭터를 특징으로 합니다. 키보드, 컨트롤러, 터치 디바이스의 가상 조이스틱으로 제어할 수 있는 이동은 두 축으로 제한되어 있어, 앞, 뒤, 점프만 가능합니다. 레벨에는 캐릭터가 점프해 다닐 수 있는 턱이 다수 있습니다."
+		},
+		{
+			"Language": "zh-CN",
+			"Text": "该横版过关类游戏模板拥有一个可活动的骨架网格体角色，摄像机位于角色的一侧。您可以通过键盘、手柄或触摸式设备上的虚拟摇杆来控制角色移动，但仅限于两个方向：前进、后退及跳跃。关卡中有许多岩架，角色可以在它们之间跳跃。"
 		},
 		{
 			"Language": "ja",
-			"Text": "The side scroller pack features an animated skeletal mesh character with a camera positioned at the character's side. Movement, which can be controlled with either keyboard, controller or virtual joystick on a touch device is limited to two dimensions: forward, backward and jumping. A number of ledges are featured in the level and the character can jump between them."
+			"Text": "サイドスクロールテンプレートにはアニメーションを使用するスケルタルメッシュのキャラクタが存在し、カメラはキャラクタの側面に存在します。移動はキーボード、コントローラー、タッチデバイスでのバーチャルジョイスティックで操作可能で、二次元に制限されています。前進、後退とジャンプだけです。レベルには足場があり、それらの間をジャンプすることができます。"
 		}
 
 	],
@@ -36,25 +44,13 @@
 		{
 			"Language": "en",
 			"Text": "Animation Sequence, Animation Blueprint, Skeletal Mesh, Blend Space 1D, Skeleton"
-		},
-		{
-			"Language": "ko",
-			"Text": "애니메이션 시퀀스, 애니메이션 블루프린트, 스켈레탈 메시, 블렌드 스페이스 1D, 스켈레톤"
-		},
-		{
-			"Language": "ja",
-			"Text": "アニメーションシーケンス・アニメーションブループリント・スケルタルメッシュ・ブレンドスペース１D・スケルトン"
 		}
-	],	
+	],
 	"SearchTags":
 	[
 		{
 			"Language": "en",
 			"Text": "Side,Scroller,Platform"
-		},
-		{
-			"Language": "ko",
-			"Text": "횡,스크롤,플랫폼"
 		}
 	],	
 	"ClassTypes": "CameraComponent, SpringArmComponent, PaperFlipbook, PaperCharacter, GameMode",
diff --git a/Templates/TP_SideScrollerBP/Config/TemplateDefs.ini b/Templates/TP_SideScrollerBP/Config/TemplateDefs.ini
--- a/Templates/TP_SideScrollerBP/Config/TemplateDefs.ini
+++ b/Templates/TP_SideScrollerBP/Config/TemplateDefs.ini
@@ -3,7 +3,13 @@ LocalizedDisplayNames=(Language="en",Text="Side Scroller")
 LocalizedDescriptions=(Language="en",Text="This side scroller template features an animated skeletal mesh character with a camera positioned at the character's side. Movement, which can be controlled with either keyboard, controller or virtual joystick on a touch device is limited to two dimensions: forward, backward and jumping. A number of ledges are featured in the level and the character can jump between them.")
 
 LocalizedDisplayNames=(Language="ko",Text="횡스크롤")
-LocalizedDescriptions=(Language="ko",Text="횡스크롤 템플릿에는 측면에 카메라가 위치한 애니메이션 적용 스켈레탈 메시 캐릭터를 특징으로 합니다. 키보드, 콘트롤러, 터치 디바이스의 가상 조이스틱으로 제어할 수 있는 이동은 두 축으로 제한되어 있어, 앞, 뒤, 점프만 가능합니다. 레벨에는 캐릭터가 점프해 다닐 수 있는 턱이 다수 있습니다.")
+LocalizedDescriptions=(Language="ko",Text="횡스크롤 템플릿에는 카메라가 측면에 위치한 애니메이션 적용 스켈레탈 메시 캐릭터를 특징으로 합니다. 키보드, 컨트롤러, 터치 디바이스의 가상 조이스틱으로 제어할 수 있는 이동은 두 축으로 제한되어 있어, 앞, 뒤, 점프만 가능합니다. 레벨에는 캐릭터가 점프해 다닐 수 있는 턱이 다수 있습니다.")
+
+LocalizedDisplayNames=(Language="ja",Text="サイドスクロール")
+LocalizedDescriptions=(Language="ja",Text="サイドスクロールテンプレートにはアニメーションを使用するスケルタルメッシュのキャラクタが存在し、カメラはキャラクタの側面に存在します。移動はキーボード、コントローラー、タッチデバイスでのバーチャルジョイスティックで操作可能で、二次元に制限されています。前進、後退とジャンプだけです。レベルには足場があり、それらの間をジャンプすることができます。")
+
+LocalizedDisplayNames=(Language="zh-CN",Text="横版过关游戏")
+LocalizedDescriptions=(Language="zh-CN",Text="该横版过关类游戏模板拥有一个可活动的骨架网格体角色，摄像机位于角色的一侧。您可以通过键盘、手柄或触摸式设备上的虚拟摇杆来控制角色移动，但仅限于两个方向：前进、后退及跳跃。关卡中有许多岩架，角色可以在它们之间跳跃。")
 
 ClassTypes=CameraComponent, SpringArmComponent, PaperFlipbook, PaperCharacter, GameMode
 AssetTypes=Animation Sequence, Animation Blueprint, Skeletal Mesh, Blend Space 1D, Skeleton
diff --git a/Templates/TP_SideScrollerBP/manifest.json b/Templates/TP_SideScrollerBP/manifest.json
--- a/Templates/TP_SideScrollerBP/manifest.json
+++ b/Templates/TP_SideScrollerBP/manifest.json
@@ -10,6 +10,10 @@
 			"Language": "ko",
 			"Text": "횡스크롤"
 		},
+		{
+			"Language": "zh-CN",
+			"Text": "横版过关游戏"
+		},
 		{
 			"Language": "ja",
 			"Text": "横スクロール"
@@ -23,11 +27,15 @@
 		},
 		{
 			"Language": "ko",
-			"Text": "횡스크롤 팩에는 측면에 카메라가 위치한 애니메이션 적용 스켈레탈 메시 캐릭터를 특징으로 합니다. 키보드, 콘트롤러, 터치 디바이스의 가상 조이스틱으로 제어할 수 있는 이동은 두 축으로 제한되어 있어, 앞, 뒤, 점프만 가능합니다. 레벨에는 캐릭터가 점프해 다닐 수 있는 턱이 다수 있습니다."
+			"Text": "횡스크롤 템플릿에는 카메라가 측면에 위치한 애니메이션 적용 스켈레탈 메시 캐릭터를 특징으로 합니다. 키보드, 컨트롤러, 터치 디바이스의 가상 조이스틱으로 제어할 수 있는 이동은 두 축으로 제한되어 있어, 앞, 뒤, 점프만 가능합니다. 레벨에는 캐릭터가 점프해 다닐 수 있는 턱이 다수 있습니다."
+		},
+		{
+			"Language": "zh-CN",
+			"Text": "该横版过关类游戏模板拥有一个可活动的骨架网格体角色，摄像机位于角色的一侧。您可以通过键盘、手柄或触摸式设备上的虚拟摇杆来控制角色移动，但仅限于两个方向：前进、后退及跳跃。关卡中有许多岩架，角色可以在它们之间跳跃。"
 		},
 		{
 			"Language": "ja",
-			"Text": "The side scroller pack features an animated skeletal mesh character with a camera positioned at the character's side. Movement, which can be controlled with either keyboard, controller or virtual joystick on a touch device is limited to two dimensions: forward, backward and jumping. A number of ledges are featured in the level and the character can jump between them."
+			"Text": "サイドスクロールテンプレートにはアニメーションを使用するスケルタルメッシュのキャラクタが存在し、カメラはキャラクタの側面に存在します。移動はキーボード、コントローラー、タッチデバイスでのバーチャルジョイスティックで操作可能で、二次元に制限されています。前進、後退とジャンプだけです。レベルには足場があり、それらの間をジャンプすることができます。"
 		}
 
 	],
@@ -36,14 +44,6 @@
 		{
 			"Language": "en",
 			"Text": "Animation Sequence, Animation Blueprint, Skeletal Mesh, Blend Space 1D, Skeleton"
-		},
-		{
-			"Language": "ko",
-			"Text": "애니메이션 시퀀스, 애니메이션 블루프린트, 스켈레탈 메시, 블렌드 스페이스 1D, 스켈레톤"
-		},
-		{
-			"Language": "ja",
-			"Text": "アニメーションシーケンス・アニメーションブループリント・スケルタルメッシュ・ブレンドスペース１D・スケルトン"
 		}
 	],
 	"SearchTags":
@@ -51,10 +51,6 @@
 		{
 			"Language": "en",
 			"Text": "Side,Scroller,Platform"
-		},
-		{
-			"Language": "ko",
-			"Text": "횡,스크롤,플랫폼"
 		}
 	],	
 	"ClassTypes": "CameraComponent, SpringArmComponent, PaperFlipbook, PaperCharacter, GameMode",
diff --git a/Templates/TP_ThirdPerson/Config/TemplateDefs.ini b/Templates/TP_ThirdPerson/Config/TemplateDefs.ini
--- a/Templates/TP_ThirdPerson/Config/TemplateDefs.ini
+++ b/Templates/TP_ThirdPerson/Config/TemplateDefs.ini
@@ -5,6 +5,12 @@ LocalizedDescriptions=(Language="en",Text="The third person template features a
 LocalizedDisplayNames=(Language="ko",Text="삼인칭")
 LocalizedDescriptions=(Language="ko",Text="삼인칭 템플릿은 카메라가 머리위 약간 뒤쪽에 위치한 플레이가능 캐릭터를 특징으로 합니다. 마우스, 키보드, 터치 디바이스의 가상 조이스틱으로 캐릭터가 이동함에 따라 카메라는 캐릭터를 따라다니며 어깨위 시점으로 봅니다. 주인공 캐릭터를 강조하는 시점으로, 액션이나 어드벤처 게임에 많이 쓰입니다. 레벨에는 플레이어가 점프할 수 있는 오브젝트가 여럿 있습니다.")
 
+LocalizedDisplayNames=(Language="ja",Text="サードパーソン")
+LocalizedDescriptions=(Language="ja",Text="サードパーソンテンプレートにはプレイアブルキャラクタが存在し、カメラがキャラクタの背後の少し上方にあります。キーボード、コントローラー、タッチデバイスでのバーチャルジョイスティックでキャラクタが移動すると、カメラは肩越しのビューで追従します。この視点はメインキャラクタを強調するのもでアクションアドベンチャーゲームで人気があるものです。レベルにはプレイヤーがジャンプして乗ることができるオブジェクトが置かれています。")
+
+LocalizedDisplayNames=(Language="zh-CN",Text="第三人称游戏")
+LocalizedDescriptions=(Language="zh-CN",Text="第三人称游戏模板拥有一个可游玩的角色，摄像机位于角色身后偏上的位置。您可以通过鼠标、键盘、手柄或触摸式设备上控制角色移动，摄像机会跟随角色，通过过肩式视角观察世界。此视角突出主角，在动作及冒险类游戏中非常流行。关卡中还包含一些对象，可以让角色跳到上面。")
+
 ClassTypes=GameMode, Character, SpringArmComponent, CameraComponent, InputComponent
 AssetTypes=Animation Sequence, Animation Blueprint, Skeleton, Blend Space 1D, Skeletal Mesh
 
diff --git a/Templates/TP_ThirdPerson/manifest.json b/Templates/TP_ThirdPerson/manifest.json
--- a/Templates/TP_ThirdPerson/manifest.json
+++ b/Templates/TP_ThirdPerson/manifest.json
@@ -10,9 +10,13 @@
 			"Language": "ko",
 			"Text": "삼인칭"
 		},
+		{
+			"Language": "zh-CN",
+			"Text": "第三人称游戏"
+		},
 		{
 			"Language": "ja",
-			"Text": "アクション"
+			"Text": "サードパーソン"
 		}
 	],
 	"Description":
@@ -23,11 +27,15 @@
 		},
 		{
 			"Language": "ko",
-			"Text": "삼인칭 팩은 카메라가 머리위 약간 뒤쪽에 위치한 플레이가능 캐릭터를 특징으로 합니다. 마우스, 키보드, 터치 디바이스의 가상 조이스틱으로 캐릭터가 이동함에 따라 카메라는 캐릭터를 따라다니며 어깨위 시점으로 봅니다. 주인공 캐릭터를 강조하는 시점으로, 액션이나 어드벤처 게임에 많이 쓰입니다. 레벨에는 플레이어가 점프할 수 있는 오브젝트가 여럿 있습니다."
+			"Text": "삼인칭 템플릿은 카메라가 머리위 약간 뒤쪽에 위치한 플레이가능 캐릭터를 특징으로 합니다. 마우스, 키보드, 터치 디바이스의 가상 조이스틱으로 캐릭터가 이동함에 따라 카메라는 캐릭터를 따라다니며 어깨위 시점으로 봅니다. 주인공 캐릭터를 강조하는 시점으로, 액션이나 어드벤처 게임에 많이 쓰입니다. 레벨에는 플레이어가 점프할 수 있는 오브젝트가 여럿 있습니다."
+		},
+		{
+			"Language": "zh-CN",
+			"Text": "第三人称游戏模板拥有一个可游玩的角色，摄像机位于角色身后偏上的位置。您可以通过鼠标、键盘、手柄或触摸式设备上控制角色移动，摄像机会跟随角色，通过过肩式视角观察世界。此视角突出主角，在动作及冒险类游戏中非常流行。关卡中还包含一些对象，可以让角色跳到上面。"
 		},
 		{
 			"Language": "ja",
-			"Text": "アクションテンプレートではキャラクターを三人称視点で見下ろして操作します。キャラクターの操作にはマウス・キーボード・コントローラー・バーチャルジョイスティック・タッチパネルを使うことができます。カメラはキャラクターの上空後方から常に追従します。アクションゲームやアドベンチャーゲームに適しています。 レベルに配置されたブロックにはジャンプして登ることができます。"
+			"Text": "サードパーソンテンプレートにはプレイアブルキャラクタが存在し、カメラがキャラクタの背後の少し上方にあります。キーボード、コントローラー、タッチデバイスでのバーチャルジョイスティックでキャラクタが移動すると、カメラは肩越しのビューで追従します。この視点はメインキャラクタを強調するのもでアクションアドベンチャーゲームで人気があるものです。レベルにはプレイヤーがジャンプして乗ることができるオブジェクトが置かれています。"
 		}
 	],
 	"AssetTypes":
@@ -35,14 +43,6 @@
 		{
 			"Language": "en",
 			"Text": "Animation Sequence, Animation Blueprint, Skeleton, Blend Space 1D, Skeletal Mesh"
-		},
-		{
-			"Language": "ko",
-			"Text": "애니메이션 시퀀스, 애니메이션 블루프린트, 스켈레톤, 블렌드 스페이스 1D, 스켈레탈 메시"
-		},
-		{
-			"Language": "ja",
-			"Text": "アニメーションシーケンス・アニメーションブループリント・スケルトン・ブレンドスペース１D・スケルタルメッシュ"
 		}
 	],	
 	"SearchTags":
@@ -50,10 +50,6 @@
 		{
 			"Language": "en",
 			"Text": "Action, Adventure, Third"
-		},
-		{
-			"Language": "ko",
-			"Text": "액션, 어드벤처, 삼인칭"
 		}
 	],
 	"ClassTypes": "GameMode, Character, SpringArmComponent, CameraComponent, InputComponent",
diff --git a/Templates/TP_ThirdPersonBP/Config/TemplateDefs.ini b/Templates/TP_ThirdPersonBP/Config/TemplateDefs.ini
--- a/Templates/TP_ThirdPersonBP/Config/TemplateDefs.ini
+++ b/Templates/TP_ThirdPersonBP/Config/TemplateDefs.ini
@@ -5,6 +5,12 @@ LocalizedDescriptions=(Language="en",Text="The third person template features a
 LocalizedDisplayNames=(Language="ko",Text="삼인칭")
 LocalizedDescriptions=(Language="ko",Text="삼인칭 템플릿은 카메라가 머리위 약간 뒤쪽에 위치한 플레이가능 캐릭터를 특징으로 합니다. 마우스, 키보드, 터치 디바이스의 가상 조이스틱으로 캐릭터가 이동함에 따라 카메라는 캐릭터를 따라다니며 어깨위 시점으로 봅니다. 주인공 캐릭터를 강조하는 시점으로, 액션이나 어드벤처 게임에 많이 쓰입니다. 레벨에는 플레이어가 점프할 수 있는 오브젝트가 여럿 있습니다.")
 
+LocalizedDisplayNames=(Language="ja",Text="サードパーソン")
+LocalizedDescriptions=(Language="ja",Text="サードパーソンテンプレートにはプレイアブルキャラクタが存在し、カメラがキャラクタの背後の少し上方にあります。キーボード、コントローラー、タッチデバイスでのバーチャルジョイスティックでキャラクタが移動すると、カメラは肩越しのビューで追従します。この視点はメインキャラクタを強調するのもでアクションアドベンチャーゲームで人気があるものです。レベルにはプレイヤーがジャンプして乗ることができるオブジェクトが置かれています。")
+
+LocalizedDisplayNames=(Language="zh-CN",Text="第三人称游戏")
+LocalizedDescriptions=(Language="zh-CN",Text="第三人称游戏模板拥有一个可游玩的角色，摄像机位于角色身后偏上的位置。您可以通过鼠标、键盘、手柄或触摸式设备上控制角色移动，摄像机会跟随角色，通过过肩式视角观察世界。此视角突出主角，在动作及冒险类游戏中非常流行。关卡中还包含一些对象，可以让角色跳到上面。")
+
 ClassTypes=GameMode, Character, SpringArmComponent, CameraComponent, InputComponent
 AssetTypes=Animation Sequence, Animation Blueprint, Skeleton, Blend Space 1D, Skeletal Mesh
 
diff --git a/Templates/TP_ThirdPersonBP/manifest.json b/Templates/TP_ThirdPersonBP/manifest.json
--- a/Templates/TP_ThirdPersonBP/manifest.json
+++ b/Templates/TP_ThirdPersonBP/manifest.json
@@ -10,9 +10,13 @@
 			"Language": "ko",
 			"Text": "삼인칭"
 		},
+		{
+			"Language": "zh-CN",
+			"Text": "第三人称游戏"
+		},
 		{
 			"Language": "ja",
-			"Text": "三人称視点"
+			"Text": "サードパーソン"
 		}
 	],
 	"Description":
@@ -23,11 +27,15 @@
 		},
 		{
 			"Language": "ko",
-			"Text": "삼인칭 팩은 카메라가 머리위 약간 뒤쪽에 위치한 플레이가능 캐릭터를 특징으로 합니다. 마우스, 키보드, 터치 디바이스의 가상 조이스틱으로 캐릭터가 이동함에 따라 카메라는 캐릭터를 따라다니며 어깨위 시점으로 봅니다. 주인공 캐릭터를 강조하는 시점으로, 액션이나 어드벤처 게임에 많이 쓰입니다. 레벨에는 플레이어가 점프할 수 있는 오브젝트가 여럿 있습니다."
+			"Text": "삼인칭 템플릿은 카메라가 머리위 약간 뒤쪽에 위치한 플레이가능 캐릭터를 특징으로 합니다. 마우스, 키보드, 터치 디바이스의 가상 조이스틱으로 캐릭터가 이동함에 따라 카메라는 캐릭터를 따라다니며 어깨위 시점으로 봅니다. 주인공 캐릭터를 강조하는 시점으로, 액션이나 어드벤처 게임에 많이 쓰입니다. 레벨에는 플레이어가 점프할 수 있는 오브젝트가 여럿 있습니다."
+		},
+		{
+			"Language": "zh-CN",
+			"Text": "第三人称游戏模板拥有一个可游玩的角色，摄像机位于角色身后偏上的位置。您可以通过鼠标、键盘、手柄或触摸式设备上控制角色移动，摄像机会跟随角色，通过过肩式视角观察世界。此视角突出主角，在动作及冒险类游戏中非常流行。关卡中还包含一些对象，可以让角色跳到上面。"
 		},
 		{
 			"Language": "ja",
-			"Text": "三人称視点テンプレートではキャラクターを三人称視点で見下ろして操作します。キャラクターの操作にはマウス・キーボード・コントローラー・バーチャルジョイスティック・タッチパネルを使うことができます。カメラはキャラクターの上空後方から常に追従します。アクションゲームやアドベンチャーゲームに適しています。 レベルに配置されたブロックにはジャンプして登ることができます。"
+			"Text": "サードパーソンテンプレートにはプレイアブルキャラクタが存在し、カメラがキャラクタの背後の少し上方にあります。キーボード、コントローラー、タッチデバイスでのバーチャルジョイスティックでキャラクタが移動すると、カメラは肩越しのビューで追従します。この視点はメインキャラクタを強調するのもでアクションアドベンチャーゲームで人気があるものです。レベルにはプレイヤーがジャンプして乗ることができるオブジェクトが置かれています。"
 		}
 	],
 	"AssetTypes":
@@ -35,25 +43,13 @@
 		{
 			"Language": "en",
 			"Text": "Animation Sequence, Animation Blueprint, Skeleton, Blend Space 1D, Skeletal Mesh"
-		},
-		{
-			"Language": "ko",
-			"Text": "애니메이션 시퀀스, 애니메이션 블루프린트, 스켈레톤, 블렌드 스페이스 1D, 스켈레탈 메시"
-		},
-		{
-			"Language": "ja",
-			"Text": "アニメーションシーケンス・アニメーションブループリント・スケルトン・ブレンドスペース１D・スケルタルメッシュ"
 		}
-	],
+	],	
 	"SearchTags":
 	[
 		{
 			"Language": "en",
 			"Text": "Action, Adventure, Third"
-		},
-		{
-			"Language": "ko",
-			"Text": "액션, 어드벤처, 삼인칭"
 		}
 	],
 	"ClassTypes": "GameMode, Character, SpringArmComponent, CameraComponent, InputComponent",
diff --git a/Templates/TP_TopDown/Config/TemplateDefs.ini b/Templates/TP_TopDown/Config/TemplateDefs.ini
--- a/Templates/TP_TopDown/Config/TemplateDefs.ini
+++ b/Templates/TP_TopDown/Config/TemplateDefs.ini
@@ -3,7 +3,13 @@ LocalizedDisplayNames=(Language="en",Text="Top Down")
 LocalizedDescriptions=(Language="en",Text="The top down template features a playable character where the camera is positioned behind and above the avatar at greater distance than in third person perspective. The player is controlled by using the mouse or touchscreen to click on the required destination and uses the navigation system to facilitate the character movement. This perspective is often used in action role playing games.")
 
 LocalizedDisplayNames=(Language="ko",Text="내려보기")
-LocalizedDescriptions=(Language="ko",Text="내려보기 템플릿은 뒤쪽으로 삼인칭 시점보다 먼 거리에 있는 아바타 위에 카메라가 위치한 플레이가능 캐릭터를 특징으로 합니다. 플레이어 제어는 마우스나 터치스크린으로 원하는 목적지를 클릭하는 식으로 이루어지며, 원활한 캐릭터 이동을 위해 내비게이션 시스템을 사용합니다. 이 시점은 액션 롤플레잉 게임에 자주 사용됩니다.")
+LocalizedDescriptions=(Language="ko",Text="내려보기 템플릿은 플레이어 캐릭터 위 삼인칭 시점보다 멀리 카메라가 위치합니다. 플레이어 제어는 마우스나 터치스크린으로 원하는 목적지를 클릭하는 식으로 이루어지며, 원활한 캐릭터 이동을 위해 내비게이션 시스템을 사용합니다. 이 시점은 액션 롤플레잉 게임에 자주 쓰입니다.")
+
+LocalizedDisplayNames=(Language="ja",Text="トップダウン")
+LocalizedDescriptions=(Language="ja",Text="トップダウンテンプレートにはプレイアブルキャラクタが存在し、カメラはキャラクタの背後、サードパーソンテンプレートより遠い上方にあります。マウスやタッチスクリーンで行き先をクリックすることでキャラクタを操作することができます。キャラクタはナビゲーションシステムを利用して移動します。この視点はアクションロールプレイングゲームでよく使用されています。")
+
+LocalizedDisplayNames=(Language="zh-CN",Text="俯视角游戏")
+LocalizedDescriptions=(Language="zh-CN",Text="俯视角游戏模板拥有一个可游玩角色，摄像机位于角色的后上方，但位置比第三人称视角要高得多。玩家可以使用鼠标或触摸屏点击目的地，或使用导航系统来辅助角色移动。此视角常用于动作角色扮演类游戏。")
 
 ClassTypes=GameMode, Character, CameraComponent, SpringArmComponent, PlayerController
 AssetTypes=Animation Sequence, Skeletal Mesh, Animation Blueprint,Skeleton, Blend Space 1D
diff --git a/Templates/TP_TopDown/manifest.json b/Templates/TP_TopDown/manifest.json
--- a/Templates/TP_TopDown/manifest.json
+++ b/Templates/TP_TopDown/manifest.json
@@ -10,6 +10,10 @@
 			"Language": "ko",
 			"Text": "내려보기"
 		},
+		{
+			"Language": "zh-CN",
+			"Text": "俯视角游戏"
+		},
 		{
 			"Language": "ja",
 			"Text": "トップダウン"
@@ -23,11 +27,15 @@
 		},
 		{
 			"Language": "ko",
-			"Text": "내려보기 팩은 뒤쪽으로 삼인칭 시점보다 먼 거리에 있는 아바타 위에 카메라가 위치한 플레이가능 캐릭터를 특징으로 합니다. 플레이어 제어는 마우스나 터치스크린으로 원하는 목적지를 클릭하는 식으로 이루어지며, 원활한 캐릭터 이동을 위해 내비게이션 시스템을 사용합니다. 이 시점은 액션 롤플레잉 게임에 자주 사용됩니다."
+			"Text": "내려보기 템플릿은 플레이어 캐릭터 위 삼인칭 시점보다 멀리 카메라가 위치합니다. 플레이어 제어는 마우스나 터치스크린으로 원하는 목적지를 클릭하는 식으로 이루어지며, 원활한 캐릭터 이동을 위해 내비게이션 시스템을 사용합니다. 이 시점은 액션 롤플레잉 게임에 자주 쓰입니다."
+		},
+		{
+			"Language": "zh-CN",
+			"Text": "俯视角游戏模板拥有一个可游玩角色，摄像机位于角色的后上方，但位置比第三人称视角要高得多。玩家可以使用鼠标或触摸屏点击目的地，或使用导航系统来辅助角色移动。此视角常用于动作角色扮演类游戏。"
 		},
 		{
 			"Language": "ja",
-			"Text": "The top down pack features a playable character where the camera is positioned behind and above the avatar at greater distance than in third person perspective. The player is controlled by using the mouse or touchscreen to click on the required destination and uses the navigation system to facilitate the character movement. This perspective is often used in action role playing games."
+			"Text": "トップダウンテンプレートにはプレイアブルキャラクタが存在し、カメラはキャラクタの背後、サードパーソンテンプレートより遠い上方にあります。マウスやタッチスクリーンで行き先をクリックすることでキャラクタを操作することができます。キャラクタはナビゲーションシステムを利用して移動します。この視点はアクションロールプレイングゲームでよく使用されています。"
 		}
 	],
 	"AssetTypes":
@@ -35,14 +43,6 @@
 		{
 			"Language": "en",
 			"Text": "Animation Sequence, Skeletal Mesh, Animation Blueprint,Skeleton, Blend Space 1D"
-		},
-		{
-			"Language": "ko",
-			"Text": "애니메이션 시퀀스, 스켈레탈 메시, 애니메이션 블루프린트, 스켈레톤, 블렌드 스페이스 1D"
-		},
-		{
-			"Language": "ja",
-			"Text": "アニメーションシーケンス・スケルタルメッシュ・アニメーションブループリント・スケルトン・ブレンドシェイプ１D"
 		}
 	],	
 	"SearchTags":
@@ -50,10 +50,6 @@
 		{
 			"Language": "en",
 			"Text": "Role, RPG,Role Playing,Top Down"
-		},
-		{
-			"Language": "ko",
-			"Text": "롤, RPG, 롤플레잉, 내려보기, 탑다운"
 		}
 	],
 	"ClassTypes": "GameMode, Character, CameraComponent, SpringArmComponent, PlayerController",
diff --git a/Templates/TP_TopDownBP/Config/TemplateDefs.ini b/Templates/TP_TopDownBP/Config/TemplateDefs.ini
--- a/Templates/TP_TopDownBP/Config/TemplateDefs.ini
+++ b/Templates/TP_TopDownBP/Config/TemplateDefs.ini
@@ -3,8 +3,13 @@ LocalizedDisplayNames=(Language="en",Text="Top Down")
 LocalizedDescriptions=(Language="en",Text="The top down template features a playable character where the camera is positioned behind and above the avatar at greater distance than in third person perspective. The player is controlled by using the mouse or touchscreen to click on the required destination and uses the navigation system to facilitate the character movement. This perspective is often used in action role playing games.")
 
 LocalizedDisplayNames=(Language="ko",Text="내려보기")
-LocalizedDescriptions=(Language="ko",Text="내려보기 템플릿은 뒤쪽으로 삼인칭 시점보다 먼 거리에 있는 아바타 위에 카메라가 위치한 플레이가능 캐릭터를 특징으로 합니다. 플레이어 제어는 마우스나 터치스크린으로 원하는 목적지를 클릭하는 식으로 이루어지며, 원활한 캐릭터 이동을 위해 내비게이션 시스템을 사용합니다. 이 시점은 액션 롤플레잉 게임에 자주 사용됩니다.")
+LocalizedDescriptions=(Language="ko",Text="내려보기 템플릿은 플레이어 캐릭터 위 삼인칭 시점보다 멀리 카메라가 위치합니다. 플레이어 제어는 마우스나 터치스크린으로 원하는 목적지를 클릭하는 식으로 이루어지며, 원활한 캐릭터 이동을 위해 내비게이션 시스템을 사용합니다. 이 시점은 액션 롤플레잉 게임에 자주 쓰입니다.")
 
+LocalizedDisplayNames=(Language="ja",Text="トップダウン")
+LocalizedDescriptions=(Language="ja",Text="トップダウンテンプレートにはプレイアブルキャラクタが存在し、カメラはキャラクタの背後、サードパーソンテンプレートより遠い上方にあります。マウスやタッチスクリーンで行き先をクリックすることでキャラクタを操作することができます。キャラクタはナビゲーションシステムを利用して移動します。この視点はアクションロールプレイングゲームでよく使用されています。")
+
+LocalizedDisplayNames=(Language="zh-CN",Text="俯视角游戏")
+LocalizedDescriptions=(Language="zh-CN",Text="俯视角游戏模板拥有一个可游玩角色，摄像机位于角色的后上方，但位置比第三人称视角要高得多。玩家可以使用鼠标或触摸屏点击目的地，或使用导航系统来辅助角色移动。此视角常用于动作角色扮演类游戏。")
 
 ClassTypes=GameMode, Character, CameraComponent, SpringArmComponent, PlayerController
 AssetTypes=Animation Sequence, Skeletal Mesh, Animation Blueprint,Skeleton, Blend Space 1D
diff --git a/Templates/TP_TopDownBP/manifest.json b/Templates/TP_TopDownBP/manifest.json
--- a/Templates/TP_TopDownBP/manifest.json
+++ b/Templates/TP_TopDownBP/manifest.json
@@ -10,6 +10,10 @@
 			"Language": "ko",
 			"Text": "내려보기"
 		},
+		{
+			"Language": "zh-CN",
+			"Text": "俯视角游戏"
+		},
 		{
 			"Language": "ja",
 			"Text": "トップダウン"
@@ -23,11 +27,15 @@
 		},
 		{
 			"Language": "ko",
-			"Text": "내려보기 팩은 뒤쪽으로 삼인칭 시점보다 먼 거리에 있는 아바타 위에 카메라가 위치한 플레이가능 캐릭터를 특징으로 합니다. 플레이어 제어는 마우스나 터치스크린으로 원하는 목적지를 클릭하는 식으로 이루어지며, 원활한 캐릭터 이동을 위해 내비게이션 시스템을 사용합니다. 이 시점은 액션 롤플레잉 게임에 자주 사용됩니다."
+			"Text": "내려보기 템플릿은 플레이어 캐릭터 위 삼인칭 시점보다 멀리 카메라가 위치합니다. 플레이어 제어는 마우스나 터치스크린으로 원하는 목적지를 클릭하는 식으로 이루어지며, 원활한 캐릭터 이동을 위해 내비게이션 시스템을 사용합니다. 이 시점은 액션 롤플레잉 게임에 자주 쓰입니다."
+		},
+		{
+			"Language": "zh-CN",
+			"Text": "俯视角游戏模板拥有一个可游玩角色，摄像机位于角色的后上方，但位置比第三人称视角要高得多。玩家可以使用鼠标或触摸屏点击目的地，或使用导航系统来辅助角色移动。此视角常用于动作角色扮演类游戏。"
 		},
 		{
 			"Language": "ja",
-			"Text": "The top down pack features a playable character where the camera is positioned behind and above the avatar at greater distance than in third person perspective. The player is controlled by using the mouse or touchscreen to click on the required destination and uses the navigation system to facilitate the character movement. This perspective is often used in action role playing games."
+			"Text": "トップダウンテンプレートにはプレイアブルキャラクタが存在し、カメラはキャラクタの背後、サードパーソンテンプレートより遠い上方にあります。マウスやタッチスクリーンで行き先をクリックすることでキャラクタを操作することができます。キャラクタはナビゲーションシステムを利用して移動します。この視点はアクションロールプレイングゲームでよく使用されています。"
 		}
 	],
 	"AssetTypes":
@@ -35,14 +43,6 @@
 		{
 			"Language": "en",
 			"Text": "Animation Sequence, Skeletal Mesh, Animation Blueprint,Skeleton, Blend Space 1D"
-		},
-		{
-			"Language": "ko",
-			"Text": "애니메이션 시퀀스, 스켈레탈 메시, 애니메이션 블루프린트, 스켈레톤, 블렌드 스페이스 1D"
-		},
-		{
-			"Language": "ja",
-			"Text": "アニメーションシーケンス・スケルタルメッシュ・アニメーションブループリント・スケルトン・ブレンドシェイプ１D"
 		}
 	],	
 	"SearchTags":
@@ -50,10 +50,6 @@
 		{
 			"Language": "en",
 			"Text": "Role, RPG,Role Playing,Top Down"
-		},
-		{
-			"Language": "ko",
-			"Text": "롤, RPG, 롤플레잉, 내려보기, 탑다운"
 		}
 	],
 	"ClassTypes": "GameMode, Character, CameraComponent, SpringArmComponent, PlayerController",
diff --git a/Templates/TP_TwinStick/Config/TemplateDefs.ini b/Templates/TP_TwinStick/Config/TemplateDefs.ini
--- a/Templates/TP_TwinStick/Config/TemplateDefs.ini
+++ b/Templates/TP_TwinStick/Config/TemplateDefs.ini
@@ -3,7 +3,13 @@ LocalizedDisplayNames=(Language="en",Text="Twin Stick Shooter")
 LocalizedDescriptions=(Language="en",Text="The twin stick shooter template features a playable character where the camera is positioned almost directly above the character. Control of the character is with either keyboard, controller or virtual joystick on a touch device. Movement and firing, which are both on the X/Y axis are controlled independantly allowing for firing in a different direction than which the character is facing. The level contains several cubes that are affected by the projectiles and several that are not.")
 
 LocalizedDisplayNames=(Language="ko",Text="트윈 스틱 슈팅")
-LocalizedDescriptions=(Language="ko",Text="트윈 스틱 슈팅 템플릿은 카메라가 거의 캐릭터에 붙어있는 플레이가능 캐릭터를 특징으로 합니다. 캐릭터 콘트롤은 키보드, 콘트롤러, 터치 디바이스의 가상 조이스틱 중 하나로 가능합니다. 이동과 발사 둘 다 X/Y 축에 있으며 독립적으로 제어되어 캐릭터가 향하는 방향과 다른 방향으로 발사할 수 있습니다. 레벨에는 프로젝타일에 영향받는 큐브 여럿, 받지 않는 큐브도 여럿 들어있습니다.")
+LocalizedDescriptions=(Language="ko",Text="트윈 스틱 슈팅 템플릿은 카메라가 플레이어 캐릭터에 거의 붙어있습니다. 캐릭터 제어는 키보드, 컨트롤러, 터치 디바이스의 가상 조이스틱 중 하나로 가능합니다. 이동과 발사가 X/Y 축으로 각각 독립 제어되어 캐릭터가 향하는 방향과 다른 방향으로 발사할 수 있습니다. 레벨에는 발사체에 영향받는 큐브도 받지 않는 큐브도 여럿 있습니다.")
+
+LocalizedDisplayNames=(Language="ja",Text="ツインスティックシューター")
+LocalizedDescriptions=(Language="ja",Text="ツインスティックシューターテンプレートにはプレイアブルキャラクタが存在し、カメラはほぼキャラクタの真上にあります。キーボード、コントローラー、タッチデバイスでのバーチャルジョイスティックでこのキャラクタを操作できます。移動と弾の発射は、どちらも X と Y 軸で独立に操作します。キャラクタ向きと別の方向に弾を発射することもできます。弾の影響を受けるキューブと、影響を受けないキューブがレベルに含まれています。")
+
+LocalizedDisplayNames=(Language="zh-CN",Text="双摇杆射击类游戏")
+LocalizedDescriptions=(Language="zh-CN",Text="双摇杆设计类游戏模板拥有一个可游玩角色，摄像机几乎位于角色的正后方。您可以通过键盘、手柄或触摸式设备的虚拟摇杆控制角色移动。角色可以在X/Y轴上移动和射击，且为独立控制，所以角色可以向自身面向之外的方向射击。关卡内还有一些可以被子弹影响或不能被子弹影响的立方体。")
 
 ClassTypes=GameMode, Pawn, Actor, StaticMeshComponent, ProjectileMovementComponent, StaticMeshComponent, PrimitiveComponent, CameraComponent, SpringArmComponent, InputComponent
 AssetTypes=Sound Wave, Static Mesh
diff --git a/Templates/TP_TwinStick/manifest.json b/Templates/TP_TwinStick/manifest.json
--- a/Templates/TP_TwinStick/manifest.json
+++ b/Templates/TP_TwinStick/manifest.json
@@ -10,9 +10,13 @@
 			"Language": "ko",
 			"Text": "트윈 스틱 슈팅"
 		},
+		{
+			"Language": "zh-CN",
+			"Text": "双摇杆射击类游戏"
+		},
 		{
 			"Language": "ja",
-			"Text": "スクロールシューター"
+			"Text": "ツインスティックシューター"
 		}
 	],
 	"Description":
@@ -23,11 +27,15 @@
 		},
 		{
 			"Language": "ko",
-			"Text": "트윈 스틱 슈팅 팩은 카메라가 거의 캐릭터에 붙어있는 플레이가능 캐릭터를 특징으로 합니다. 캐릭터 콘트롤은 키보드, 콘트롤러, 터치 디바이스의 가상 조이스틱 중 하나로 가능합니다. 이동과 발사 둘 다 X/Y 축에 있으며 독립적으로 제어되어 캐릭터가 향하는 방향과 다른 방향으로 발사할 수 있습니다. 레벨에는 프로젝타일에 영향받는 큐브 여럿, 받지 않는 큐브도 여럿 들어있습니다."
+			"Text": "트윈 스틱 슈팅 템플릿은 카메라가 플레이어 캐릭터에 거의 붙어있습니다. 캐릭터 제어는 키보드, 컨트롤러, 터치 디바이스의 가상 조이스틱 중 하나로 가능합니다. 이동과 발사가 X/Y 축으로 각각 독립 제어되어 캐릭터가 향하는 방향과 다른 방향으로 발사할 수 있습니다. 레벨에는 발사체에 영향받는 큐브도 받지 않는 큐브도 여럿 있습니다."
+		},
+		{
+			"Language": "zh-CN",
+			"Text": "双摇杆设计类游戏模板拥有一个可游玩角色，摄像机几乎位于角色的正后方。您可以通过键盘、手柄或触摸式设备的虚拟摇杆控制角色移动。角色可以在X/Y轴上移动和射击，且为独立控制，所以角色可以向自身面向之外的方向射击。关卡内还有一些可以被子弹影响或不能被子弹影响的立方体。"
 		},
 		{
 			"Language": "ja",
-			"Text": "The twin stick shooter pack features a playable character where the camera is positioned almost directly above the character. Control of the character is with either keyboard, controller or virtual joystick on a touch device. Movement and firing, which are both on the X/Y axis are controlled independantly allowing for firing in a different direction than which the character is facing. The level contains several cubes that are affected by the projectiles and several that are not."
+			"Text": "ツインスティックシューターテンプレートにはプレイアブルキャラクタが存在し、カメラはほぼキャラクタの真上にあります。キーボード、コントローラー、タッチデバイスでのバーチャルジョイスティックでこのキャラクタを操作できます。移動と弾の発射は、どちらも X と Y 軸で独立に操作します。キャラクタ向きと別の方向に弾を発射することもできます。弾の影響を受けるキューブと、影響を受けないキューブがレベルに含まれています。"
 		}
 	],
 	"AssetTypes":
@@ -35,14 +43,6 @@
 		{
 			"Language": "en",
 			"Text": "Sound Wave, Static Mesh"
-		},
-		{
-			"Language": "ko",
-			"Text": "사운드 웨이브, 스태틱 메시"
-		},
-		{
-			"Language": "ja",
-			"Text": "サウンドウェーブ・スタティックメッシュ"
 		}
 	],
 	"SearchTags":
@@ -50,10 +50,6 @@
 		{
 			"Language": "en",
 			"Text": "Twin,Stick,Fire,Shooter"
-		},
-		{
-			"Language": "ko",
-			"Text": "트윈,스틱,슈팅"
 		}
 	],
 	"ClassTypes": "GameMode, Pawn, Actor, StaticMeshComponent, ProjectileMovementComponent, StaticMeshComponent, PrimitiveComponent, CameraComponent, SpringArmComponent, InputComponent",
diff --git a/Templates/TP_TwinStickBP/Config/TemplateDefs.ini b/Templates/TP_TwinStickBP/Config/TemplateDefs.ini
--- a/Templates/TP_TwinStickBP/Config/TemplateDefs.ini
+++ b/Templates/TP_TwinStickBP/Config/TemplateDefs.ini
@@ -3,7 +3,13 @@ LocalizedDisplayNames=(Language="en",Text="Twin Stick Shooter")
 LocalizedDescriptions=(Language="en",Text="The twin stick shooter template features a playable character where the camera is positioned almost directly above the character. Control of the character is with either keyboard, controller or virtual joystick on a touch device. Movement and firing, which are both on the X/Y axis are controlled independantly allowing for firing in a different direction than which the character is facing. The level contains several cubes that are affected by the projectiles and several that are not.")
 
 LocalizedDisplayNames=(Language="ko",Text="트윈 스틱 슈팅")
-LocalizedDescriptions=(Language="ko",Text="트윈 스틱 슈팅 템플릿은 카메라가 거의 캐릭터에 붙어있는 플레이가능 캐릭터를 특징으로 합니다. 캐릭터 콘트롤은 키보드, 콘트롤러, 터치 디바이스의 가상 조이스틱 중 하나로 가능합니다. 이동과 발사 둘 다 X/Y 축에 있으며 독립적으로 제어되어 캐릭터가 향하는 방향과 다른 방향으로 발사할 수 있습니다. 레벨에는 프로젝타일에 영향받는 큐브 여럿, 받지 않는 큐브도 여럿 들어있습니다.")
+LocalizedDescriptions=(Language="ko",Text="트윈 스틱 슈팅 템플릿은 카메라가 플레이어 캐릭터에 거의 붙어있습니다. 캐릭터 제어는 키보드, 컨트롤러, 터치 디바이스의 가상 조이스틱 중 하나로 가능합니다. 이동과 발사가 X/Y 축으로 각각 독립 제어되어 캐릭터가 향하는 방향과 다른 방향으로 발사할 수 있습니다. 레벨에는 발사체에 영향받는 큐브도 받지 않는 큐브도 여럿 있습니다.")
+
+LocalizedDisplayNames=(Language="ja",Text="ツインスティックシューター")
+LocalizedDescriptions=(Language="ja",Text="ツインスティックシューターテンプレートにはプレイアブルキャラクタが存在し、カメラはほぼキャラクタの真上にあります。キーボード、コントローラー、タッチデバイスでのバーチャルジョイスティックでこのキャラクタを操作できます。移動と弾の発射は、どちらも X と Y 軸で独立に操作します。キャラクタ向きと別の方向に弾を発射することもできます。弾の影響を受けるキューブと、影響を受けないキューブがレベルに含まれています。")
+
+LocalizedDisplayNames=(Language="zh-CN",Text="双摇杆射击类游戏")
+LocalizedDescriptions=(Language="zh-CN",Text="双摇杆设计类游戏模板拥有一个可游玩角色，摄像机几乎位于角色的正后方。您可以通过键盘、手柄或触摸式设备的虚拟摇杆控制角色移动。角色可以在X/Y轴上移动和射击，且为独立控制，所以角色可以向自身面向之外的方向射击。关卡内还有一些可以被子弹影响或不能被子弹影响的立方体。")
 
 ClassTypes=GameMode, Pawn, Actor, StaticMeshComponent, ProjectileMovementComponent, StaticMeshComponent, PrimitiveComponent, CameraComponent, SpringArmComponent, InputComponent
 AssetTypes=Sound Wave, Static Mesh
diff --git a/Templates/TP_TwinStickBP/manifest.json b/Templates/TP_TwinStickBP/manifest.json
--- a/Templates/TP_TwinStickBP/manifest.json
+++ b/Templates/TP_TwinStickBP/manifest.json
@@ -10,9 +10,13 @@
 			"Language": "ko",
 			"Text": "트윈 스틱 슈팅"
 		},
+		{
+			"Language": "zh-CN",
+			"Text": "双摇杆射击类游戏"
+		},
 		{
 			"Language": "ja",
-			"Text": "スクロールシューター"
+			"Text": "ツインスティックシューター"
 		}
 	],
 	"Description":
@@ -23,11 +27,15 @@
 		},
 		{
 			"Language": "ko",
-			"Text": "트윈 스틱 슈팅 팩은 카메라가 거의 캐릭터에 붙어있는 플레이가능 캐릭터를 특징으로 합니다. 캐릭터 콘트롤은 키보드, 콘트롤러, 터치 디바이스의 가상 조이스틱 중 하나로 가능합니다. 이동과 발사 둘 다 X/Y 축에 있으며 독립적으로 제어되어 캐릭터가 향하는 방향과 다른 방향으로 발사할 수 있습니다. 레벨에는 프로젝타일에 영향받는 큐브 여럿, 받지 않는 큐브도 여럿 들어있습니다."
+			"Text": "트윈 스틱 슈팅 템플릿은 카메라가 플레이어 캐릭터에 거의 붙어있습니다. 캐릭터 제어는 키보드, 컨트롤러, 터치 디바이스의 가상 조이스틱 중 하나로 가능합니다. 이동과 발사가 X/Y 축으로 각각 독립 제어되어 캐릭터가 향하는 방향과 다른 방향으로 발사할 수 있습니다. 레벨에는 발사체에 영향받는 큐브도 받지 않는 큐브도 여럿 있습니다."
+		},
+		{
+			"Language": "zh-CN",
+			"Text": "双摇杆设计类游戏模板拥有一个可游玩角色，摄像机几乎位于角色的正后方。您可以通过键盘、手柄或触摸式设备的虚拟摇杆控制角色移动。角色可以在X/Y轴上移动和射击，且为独立控制，所以角色可以向自身面向之外的方向射击。关卡内还有一些可以被子弹影响或不能被子弹影响的立方体。"
 		},
 		{
 			"Language": "ja",
-			"Text": "The twin stick shooter pack features a playable character where the camera is positioned almost directly above the character. Control of the character is with either keyboard, controller or virtual joystick on a touch device. Movement and firing, which are both on the X/Y axis are controlled independantly allowing for firing in a different direction than which the character is facing. The level contains several cubes that are affected by the projectiles and several that are not."
+			"Text": "ツインスティックシューターテンプレートにはプレイアブルキャラクタが存在し、カメラはほぼキャラクタの真上にあります。キーボード、コントローラー、タッチデバイスでのバーチャルジョイスティックでこのキャラクタを操作できます。移動と弾の発射は、どちらも X と Y 軸で独立に操作します。キャラクタ向きと別の方向に弾を発射することもできます。弾の影響を受けるキューブと、影響を受けないキューブがレベルに含まれています。"
 		}
 	],
 	"AssetTypes":
@@ -35,14 +43,6 @@
 		{
 			"Language": "en",
 			"Text": "Sound Wave, Static Mesh"
-		},
-		{
-			"Language": "ko",
-			"Text": "사운드 웨이브, 스태틱 메시"
-		},
-		{
-			"Language": "ja",
-			"Text": "サウンドウェーブ・スタティックメッシュ"
 		}
 	],
 	"SearchTags":
@@ -50,10 +50,6 @@
 		{
 			"Language": "en",
 			"Text": "Twin,Stick,Fire,Shooter"
-		},
-		{
-			"Language": "ko",
-			"Text": "트윈,스틱,슈팅"
 		}
 	],
 	"ClassTypes": "GameMode, Pawn, Actor, StaticMeshComponent, ProjectileMovementComponent, StaticMeshComponent, PrimitiveComponent, CameraComponent, SpringArmComponent, InputComponent",
diff --git a/Templates/TP_Vehicle/Config/TemplateDefs.ini b/Templates/TP_Vehicle/Config/TemplateDefs.ini
--- a/Templates/TP_Vehicle/Config/TemplateDefs.ini
+++ b/Templates/TP_Vehicle/Config/TemplateDefs.ini
@@ -3,7 +3,13 @@ LocalizedDisplayNames=(Language="en",Text="Vehicle")
 LocalizedDescriptions=(Language="en",Text="The vehicle template features a simple physics-driven vehicle. Vehicle movement can be controlled with either keyboard, controller or virtual joystick on a touch device. The template features 2 camera views - On is positioned above and behind the vehicle, whilst the other is positioned inside. When using the inside view the mouse or touchscreen or a HMD device. There is also a HUD both in and out of the car that shows the current gear and speed. The level features a smooth ramp, a slightly irregluar one, some objects that can be knocked over and a post which cannot.")
 
 LocalizedDisplayNames=(Language="ko",Text="비히클")
-LocalizedDescriptions=(Language="ko",Text="비히클 템플릿은 단순한 물리 기반 비히클을 특징으로 합니다. 비히클 이동은 키보드, 콘트롤러, 터치 디바이스의 가상 조이스틱 등 어느 것으로도 제어 가능합니다. 템플릿에는 두 가지 카메라 뷰가 제공되는데, 하나는 비히클 상단 뒤쪽에, 다른 하나는 안쪽에 위치해 있습니다. 안쪽에 있을 때는 마우스나 터치스크린이나 HMD 디바이스를 사용하여 둘러볼 수 있습니다. 차량 안에도 밖에도 현재 기어와 속도를 나타내는 HUD 가 있습니다. 레벨에는 약간 불규칙한 부드러운 경사로, 넘어뜨릴 수 있는 오브젝트와 넘어뜨릴 수 없는 기둥이 있습니다.")
+LocalizedDescriptions=(Language="ko",Text="비히클 템플릿은 단순한 물리 기반 비히클을 특징으로 합니다. 비히클 이동은 키보드, 컨트롤러, 터치 디바이스의 가상 조이스틱 등 어느 것으로도 제어 가능합니다. 템플릿에는 두 가지 카메라 뷰가 제공되는데, 하나는 비히클 상단 뒤쪽에, 다른 하나는 안쪽에 위치해 있습니다. 안쪽에 있을 때는 마우스나 터치스크린이나 HMD 디바이스를 사용하여 주변을 둘러볼 수 있습니다. 차량 안에도 밖에도 현재 기어와 속도를 나타내는 HUD 가 있습니다. 레벨에는 약간 불규칙한 부드러운 경사로, 넘어뜨릴 수 있는 오브젝트와 넘어뜨릴 수 없는 기둥이 있습니다.")
+
+LocalizedDisplayNames=(Language="ja",Text="ビークル")
+LocalizedDescriptions=(Language="ja",Text="ビークルテンプレートにはシンプルな物理駆動のビークルが含まれています。キーボード、コントローラー、タッチデバイスでのバーチャルジョイスティックでビークルを操作できます。このテンプレートには 2 方向からのカメラビューが用意されています。後方真上からのビューとビークル内部からのビューです。内部からのカメラを使用している時はマウス、タッチスクリーン、HMD デバイスで周りを見ることができます。また現在のギアと速度を示す HUD も用意されています。レベルにはスムーズな坂、多少不規則な坂、倒すことのできるオブジェクトとそうではない柱があります。")
+
+LocalizedDisplayNames=(Language="zh-CN",Text="载具类游戏")
+LocalizedDescriptions=(Language="zh-CN",Text="载具类游戏模板拥有一辆简单的物理驱动载具。您可通过键盘、手柄或触摸式设备上的虚拟摇杆来控制载具移动。该模板有2种摄像机视角——一种位于载具的后上方，而另一种位于载具内部。在使用内部视角时，您可通过鼠标、触摸屏或HMD设备观察四周。。载具内外还各有一个HUD，以显示当前装备和速度。关卡内含有一个光滑的斜坡和一个略不规则的斜坡，一些可以被撞飞的对象以及无法被撞开的柱子。")
 
 ClassTypes=GameMode, WheeledVehicle, HUD, Font, SpringArmComponent, CameraComponent, TextRenderComponent, InputComponent, VehicleWheel
 AssetTypes=Static Mesh, Animation Blueprint, Physics Asset, Physics Material, Skelton, Skeletal Mesh, Material, Material instance, Texture
diff --git a/Templates/TP_Vehicle/manifest.json b/Templates/TP_Vehicle/manifest.json
--- a/Templates/TP_Vehicle/manifest.json
+++ b/Templates/TP_Vehicle/manifest.json
@@ -10,6 +10,10 @@
 			"Language": "ko",
 			"Text": "비히클"
 		},
+		{
+			"Language": "zh-CN",
+			"Text": "载具类游戏"
+		},
 		{
 			"Language": "ja",
 			"Text": "レース"
@@ -23,11 +27,15 @@
 		},
 		{
 			"Language": "ko",
-			"Text": "비히클 팩은 단순한 물리 기반 비히클을 특징으로 합니다. 비히클 이동은 키보드, 콘트롤러, 터치 디바이스의 가상 조이스틱 등 어느 것으로도 제어 가능합니다. 템플릿에는 두 가지 카메라 뷰가 제공되는데, 하나는 비히클 상단 뒤쪽에, 다른 하나는 안쪽에 위치해 있습니다. 안쪽에 있을 때는 마우스나 터치스크린이나 HMD 디바이스를 사용하여 둘러볼 수 있습니다. 차량 안에도 밖에도 현재 기어와 속도를 나타내는 HUD 가 있습니다. 레벨에는 약간 불규칙한 부드러운 경사로, 넘어뜨릴 수 있는 오브젝트와 넘어뜨릴 수 없는 기둥이 있습니다."
+			"Text": "비히클 템플릿은 단순한 물리 기반 비히클을 특징으로 합니다. 비히클 이동은 키보드, 컨트롤러, 터치 디바이스의 가상 조이스틱 등 어느 것으로도 제어 가능합니다. 템플릿에는 두 가지 카메라 뷰가 제공되는데, 하나는 비히클 상단 뒤쪽에, 다른 하나는 안쪽에 위치해 있습니다. 안쪽에 있을 때는 마우스나 터치스크린이나 HMD 디바이스를 사용하여 주변을 둘러볼 수 있습니다. 차량 안에도 밖에도 현재 기어와 속도를 나타내는 HUD 가 있습니다. 레벨에는 약간 불규칙한 부드러운 경사로, 넘어뜨릴 수 있는 오브젝트와 넘어뜨릴 수 없는 기둥이 있습니다."
+		},
+		{
+			"Language": "zh-CN",
+			"Text": "载具类游戏模板拥有一辆简单的物理驱动载具。您可通过键盘、手柄或触摸式设备上的虚拟摇杆来控制载具移动。该模板有2种摄像机视角——一种位于载具的后上方，而另一种位于载具内部。在使用内部视角时，您可通过鼠标、触摸屏或HMD设备观察四周。。载具内外还各有一个HUD，以显示当前装备和速度。关卡内含有一个光滑的斜坡和一个略不规则的斜坡，一些可以被撞飞的对象以及无法被撞开的柱子。"
 		},
 		{
 			"Language": "ja",
-			"Text": "The vehicle pack features a simple physics-driven vehicle. Vehicle movement can be controlled with either keyboard, controller or virtual joystick on a touch device. The template features 2 camera views - On is positioned above and behind the vehicle, whilst the other is positioned inside. When using the inside view the mouse or touchscreen or a HMD device. There is also a HUD both in and out of the car that shows the current gear and speed. The level features a smooth ramp, a slightly irregluar one, some objects that can be knocked over and a post which cannot."
+			"Text": "ビークルテンプレートにはシンプルな物理駆動のビークルが含まれています。キーボード、コントローラー、タッチデバイスでのバーチャルジョイスティックでビークルを操作できます。このテンプレートには 2 方向からのカメラビューが用意されています。後方真上からのビューとビークル内部からのビューです。内部からのカメラを使用している時はマウス、タッチスクリーン、HMD デバイスで周りを見ることができます。また現在のギアと速度を示す HUD も用意されています。レベルにはスムーズな坂、多少不規則な坂、倒すことのできるオブジェクトとそうではない柱があります。"
 		}
 	],
 	"AssetTypes":
@@ -35,14 +43,6 @@
 		{
 			"Language": "en",
 			"Text": "Static Mesh, Animation Blueprint, Physics Asset, Physics Material, Skelton, Skeletal Mesh, Material, Material instance, Texture"
-		},
-		{
-			"Language": "ko",
-			"Text": "스태틱 메시, 애니메이션 블루프린트, 피직스 애셋, 피직스 머티리얼, 스켈레톤, 스켈레탈 메시, 머티리얼, 머티리얼 인스턴스, 텍스처"
-		},
-		{
-			"Language": "ja",
-			"Text": "スタティックメッシュ・アニメーションブループリント・フィジックス・フィジカルマテリアル・スケルトン・スケルタルメッシュ・マテリアル・マテリアルインスタンス・テクスチャ"
 		}
 	],
 	"SearchTags":
@@ -50,12 +50,8 @@
 		{
 			"Language": "en",
 			"Text": "Tyre,Wheel,Racing,Vehicle,Car"
-		},
-		{
-			"Language": "ko",
-			"Text": "타이어,휠,레이싱,비히클,차,경주"
 		}
-	],	
+	],
 	"ClassTypes": "GameMode, WheeledVehicle, HUD, Font, SpringArmComponent, CameraComponent, TextRenderComponent, InputComponent, VehicleWheel",
 	"Category": "CodeFeature",
 	"FocusAsset": "/Game/VehicleCPP/VehicleOverview.VehicleOverview",
diff --git a/Templates/TP_VehicleAdv/Config/TemplateDefs.ini b/Templates/TP_VehicleAdv/Config/TemplateDefs.ini
--- a/Templates/TP_VehicleAdv/Config/TemplateDefs.ini
+++ b/Templates/TP_VehicleAdv/Config/TemplateDefs.ini
@@ -3,7 +3,13 @@ LocalizedDisplayNames=(Language="en",Text="VehicleAdvanced")
 LocalizedDescriptions=(Language="en",Text="The vehicle advanced template features a more complex vehicle than the regular vehicle template. Vehicle movement can be controlled with either keyboard, controller or virtual joystick on a touch device. The template features 2 camera views - One is positioned above and behind the vehicle, whilst the other is positioned inside. When using the inside view you can look around using the mouse or touchscreen or a HMD device. The vehicle itself features working double wishbone suspension. There is also an engine sound in this template. There is also a HUD both in and out of the car that shows the current gear and speed. The level contains a landscape with rough terrain as well as a simple track that features bumps, ramps and loop the loops.")
 
 LocalizedDisplayNames=(Language="ko",Text="비히클 고급")
-LocalizedDescriptions=(Language="ko",Text="비히클 고급 템플릿은 일반 비히클 템플릿보다 복잡한 비히클이 특징입니다. 비히클 이동은 키보드, 콘트롤러, 터치 디바이스의 가상 조이스틱 등 어느 것으로도 제어 가능합니다. 템플릿에는 두 가지 카메라 뷰가 제공되는데, 하나는 비히클 상단 뒤쪽에, 다른 하나는 안쪽에 위치해 있습니다. 안쪽에 있을 때는 마우스나 터치스크린이나 HMD 디바이스를 사용하여 둘러볼 수 있습니다. 비히클 자체에는 더블 위시본 서스펜션은 물론 엔진 사운드도 구현되어 있습니다. 차량 안에도 밖에도 현재 기어와 속도를 나타내는 HUD 가 있습니다. 레벨에는 거친 지형의 랜드스케이프는 물론, 방지턱, 경사로, 루프로 된 간단한 트랙도 포함되어 있습니다.")
+LocalizedDescriptions=(Language="ko",Text="비히클 고급 템플릿은 일반 비히클 템플릿보다 복잡한 비히클이 특징입니다. 비히클 이동은 키보드, 컨트롤러, 터치 디바이스의 가상 조이스틱 등 어느 것으로도 제어 가능합니다. 템플릿에는 두 가지 카메라 뷰가 제공되는데, 하나는 비히클 상단 뒤쪽에, 다른 하나는 안쪽에 위치해 있습니다. 안쪽에 있을 때는 마우스나 터치스크린이나 HMD 디바이스를 사용하여 둘러볼 수 있습니다. 비히클 자체에는 더블 위시본 서스펜션은 물론 엔진 사운드도 구현되어 있습니다. 차량 안에도 밖에도 현재 기어와 속도를 나타내는 HUD 가 있습니다. 레벨에는 거친 지형의 랜드스케이프는 물론, 방지턱, 경사로, 간단한 루프 트랙도 포함되어 있습니다.")
+
+LocalizedDisplayNames=(Language="ja",Text="ビークル応用")
+LocalizedDescriptions=(Language="ja",Text="ビークル応用テンプレートには通常のビークルテンプレートより複雑なビークルが含まれています。キーボード、コントローラー、タッチデバイスでのバーチャルジョイスティックでビークルを操作できます。このテンプレートには 2 方向からのカメラビューが用意されています。後方真上からのビューとビークル内部からのビューです。内部からのカメラを使用している時はマウス、タッチスクリーン、HMD デバイスで周りを見ることができます。このビークルにはダブルウィッシュボーン式サスペンションが含まれています。またエンジン音もテンプレートにあります。また現在のギアと速度を示す HUD も用意されています。レベルには荒地を含む地形と、隆起、坂、ループを含む簡単なレーストラックがあります。")
+
+LocalizedDisplayNames=(Language="zh-CN",Text="高级载具类游戏")
+LocalizedDescriptions=(Language="zh-CN",Text="高级载具类游戏模板含有一辆比普通载具类游戏模板更复杂的载具。您可通过键盘、手柄或触摸式设备上的虚拟摇杆来控制载具移动。该模板有2种摄像机视角——一种位于载具的后上方，而另一种位于载具内部。在使用内部视角时，您可通过鼠标、触摸屏或HMD设备观察四周。载具本身具有双横臂式悬架。该模板还包含引擎音效。载具内外还各有一个HUD，以显示当前装备和速度。关卡内含有一片粗糙地形，以及一条包含隆起、斜坡以及环道的简单道路。")
 
 ClassTypes=GameMode, WheeledVehicle, HUD, SpringArmComponent, Font, CameraComponent, TextRenderComponent, InputComponent, PhysicalMaterial, AudioComponent
 AssetTypes=Static Mesh, Landscape Layer, Material, Instance, Physical Material, Sound Attenuation, Sound cue, Sound Wave, Texture, Skeletal Mesh, Skeleton, Physics Asset, Animation blueprint, Data Asset, Vehicle Wheel
diff --git a/Templates/TP_VehicleAdv/manifest.json b/Templates/TP_VehicleAdv/manifest.json
--- a/Templates/TP_VehicleAdv/manifest.json
+++ b/Templates/TP_VehicleAdv/manifest.json
@@ -6,6 +6,14 @@
 			"Language": "en",
 			"Text": "Vehicle Advanced"
 		},
+		{
+			"Language": "ja",
+			"Text": "ビークル応用"
+		},
+		{
+			"Language": "zh-CN",
+			"Text": "高级载具类游戏"
+		},
 		{
 			"Language": "ko",
 			"Text": "비히클 고급"
@@ -17,9 +25,17 @@
 			"Language": "en",
 			"Text": "The vehicle advanced pack features a more complex vehicle than the regular vehicle template. Vehicle movement can be controlled with either keyboard, controller or virtual joystick on a touch device. The template features 2 camera views - One is positioned above and behind the vehicle, whilst the other is positioned inside. When using the inside view you can look around using the mouse or touchscreen or a HMD device. The vehicle itself features working double wishbone suspension. There is also an engine sound in this template. There is also a HUD both in and out of the car that shows the current gear and speed. The level contains a landscape with rough terrain as well as a simple track that features bumps, ramps and loop the loops."
 		},
+		{
+			"Language": "ja",
+			"Text": "ビークル応用テンプレートには通常のビークルテンプレートより複雑なビークルが含まれています。キーボード、コントローラー、タッチデバイスでのバーチャルジョイスティックでビークルを操作できます。このテンプレートには 2 方向からのカメラビューが用意されています。後方真上からのビューとビークル内部からのビューです。内部からのカメラを使用している時はマウス、タッチスクリーン、HMD デバイスで周りを見ることができます。このビークルにはダブルウィッシュボーン式サスペンションが含まれています。またエンジン音もテンプレートにあります。また現在のギアと速度を示す HUD も用意されています。レベルには荒地を含む地形と、隆起、坂、ループを含む簡単なレーストラックがあります。"
+		},
+		{
+			"Language": "zh-CN",
+			"Text": "高级载具类游戏模板含有一辆比普通载具类游戏模板更复杂的载具。您可通过键盘、手柄或触摸式设备上的虚拟摇杆来控制载具移动。该模板有2种摄像机视角——一种位于载具的后上方，而另一种位于载具内部。在使用内部视角时，您可通过鼠标、触摸屏或HMD设备观察四周。载具本身具有双横臂式悬架。该模板还包含引擎音效。载具内外还各有一个HUD，以显示当前装备和速度。关卡内含有一片粗糙地形，以及一条包含隆起、斜坡以及环道的简单道路。"
+		},
 		{
 			"Language": "ko",
-			"Text": "비히클 고급 팩은 일반 비히클 템플릿보다 복잡한 비히클이 특징입니다. 비히클 이동은 키보드, 콘트롤러, 터치 디바이스의 가상 조이스틱 등 어느 것으로도 제어 가능합니다. 템플릿에는 두 가지 카메라 뷰가 제공되는데, 하나는 비히클 상단 뒤쪽에, 다른 하나는 안쪽에 위치해 있습니다. 안쪽에 있을 때는 마우스나 터치스크린이나 HMD 디바이스를 사용하여 둘러볼 수 있습니다. 비히클 자체에는 더블 위시본 서스펜션은 물론 엔진 사운드도 구현되어 있습니다. 차량 안에도 밖에도 현재 기어와 속도를 나타내는 HUD 가 있습니다. 레벨에는 거친 지형의 랜드스케이프는 물론, 방지턱, 경사로, 루프로 된 간단한 트랙도 포함되어 있습니다."
+			"Text": "비히클 고급 템플릿은 일반 비히클 템플릿보다 복잡한 비히클이 특징입니다. 비히클 이동은 키보드, 컨트롤러, 터치 디바이스의 가상 조이스틱 등 어느 것으로도 제어 가능합니다. 템플릿에는 두 가지 카메라 뷰가 제공되는데, 하나는 비히클 상단 뒤쪽에, 다른 하나는 안쪽에 위치해 있습니다. 안쪽에 있을 때는 마우스나 터치스크린이나 HMD 디바이스를 사용하여 둘러볼 수 있습니다. 비히클 자체에는 더블 위시본 서스펜션은 물론 엔진 사운드도 구현되어 있습니다. 차량 안에도 밖에도 현재 기어와 속도를 나타내는 HUD 가 있습니다. 레벨에는 거친 지형의 랜드스케이프는 물론, 방지턱, 경사로, 간단한 루프 트랙도 포함되어 있습니다."
 		}
 	],
 	"AssetTypes":
@@ -27,10 +43,6 @@
 		{
 			"Language": "en",
 			"Text": "Static Mesh, Landscape Layer, Material, Instance, Physical Material, Sound Attenuation, Sound cue, Sound Wave, Texture, Skeletal Mesh, Skeleton, Physics Asset, Animation blueprint, Data Asset, Vehicle Wheel"
-		},
-		{
-			"Language": "ko",
-			"Text": "스태틱 메시, 랜드스케이프 레이어, 머티리얼, 인스턴스, 피지컬 머티리얼, 사운드 감쇠, 사운드 큐, 사운드 웨이브, 텍스처, 스켈레탈 메시, 스켈레톤, 피직스 애셋, 애니메이션 블루프린트, 데이터 애셋, 비히클 휠"
 		}
 	],	
 	"SearchTags":
@@ -38,10 +50,6 @@
 		{
 			"Language": "en",
 			"Text": "Tyre,Suspension,Wheel,Racing,Vehicle,Car"
-		},
-		{
-			"Language": "ko",
-			"Text": "타이어,서스펜션,휠,레이싱,비히클,차,경주"
 		}
 	],
 	"ClassTypes": "GameMode, WheeledVehicle, HUD, SpringArmComponent, Font, CameraComponent, TextRenderComponent, InputComponent, PhysicalMaterial, AudioComponent",
diff --git a/Templates/TP_VehicleAdvBP/Config/TemplateDefs.ini b/Templates/TP_VehicleAdvBP/Config/TemplateDefs.ini
--- a/Templates/TP_VehicleAdvBP/Config/TemplateDefs.ini
+++ b/Templates/TP_VehicleAdvBP/Config/TemplateDefs.ini
@@ -3,7 +3,13 @@ LocalizedDisplayNames=(Language="en",Text="VehicleAdvanced")
 LocalizedDescriptions=(Language="en",Text="The vehicle advanced template features a more complex vehicle than the regular vehicle template. Vehicle movement can be controlled with either keyboard, controller or virtual joystick on a touch device. The template features 2 camera views - One is positioned above and behind the vehicle, whilst the other is positioned inside. When using the inside view you can look around using the mouse or touchscreen or a HMD device. The vehicle itself features working double wishbone suspension. There is also an engine sound in this template. There is also a HUD both in and out of the car that shows the current gear and speed. The level contains a landscape with rough terrain as well as a simple track that features bumps, ramps and loop the loops.")
 
 LocalizedDisplayNames=(Language="ko",Text="비히클 고급")
-LocalizedDescriptions=(Language="ko",Text="비히클 고급 템플릿은 일반 비히클 템플릿보다 복잡한 비히클이 특징입니다. 비히클 이동은 키보드, 콘트롤러, 터치 디바이스의 가상 조이스틱 등 어느 것으로도 제어 가능합니다. 템플릿에는 두 가지 카메라 뷰가 제공되는데, 하나는 비히클 상단 뒤쪽에, 다른 하나는 안쪽에 위치해 있습니다. 안쪽에 있을 때는 마우스나 터치스크린이나 HMD 디바이스를 사용하여 둘러볼 수 있습니다. 비히클 자체에는 더블 위시본 서스펜션은 물론 엔진 사운드도 구현되어 있습니다. 차량 안에도 밖에도 현재 기어와 속도를 나타내는 HUD 가 있습니다. 레벨에는 거친 지형의 랜드스케이프는 물론, 방지턱, 경사로, 루프로 된 간단한 트랙도 포함되어 있습니다.")
+LocalizedDescriptions=(Language="ko",Text="비히클 고급 템플릿은 일반 비히클 템플릿보다 복잡한 비히클이 특징입니다. 비히클 이동은 키보드, 컨트롤러, 터치 디바이스의 가상 조이스틱 등 어느 것으로도 제어 가능합니다. 템플릿에는 두 가지 카메라 뷰가 제공되는데, 하나는 비히클 상단 뒤쪽에, 다른 하나는 안쪽에 위치해 있습니다. 안쪽에 있을 때는 마우스나 터치스크린이나 HMD 디바이스를 사용하여 둘러볼 수 있습니다. 비히클 자체에는 더블 위시본 서스펜션은 물론 엔진 사운드도 구현되어 있습니다. 차량 안에도 밖에도 현재 기어와 속도를 나타내는 HUD 가 있습니다. 레벨에는 거친 지형의 랜드스케이프는 물론, 방지턱, 경사로, 간단한 루프 트랙도 포함되어 있습니다.")
+
+LocalizedDisplayNames=(Language="ja",Text="ビークル応用")
+LocalizedDescriptions=(Language="ja",Text="ビークル応用テンプレートには通常のビークルテンプレートより複雑なビークルが含まれています。キーボード、コントローラー、タッチデバイスでのバーチャルジョイスティックでビークルを操作できます。このテンプレートには 2 方向からのカメラビューが用意されています。後方真上からのビューとビークル内部からのビューです。内部からのカメラを使用している時はマウス、タッチスクリーン、HMD デバイスで周りを見ることができます。このビークルにはダブルウィッシュボーン式サスペンションが含まれています。またエンジン音もテンプレートにあります。また現在のギアと速度を示す HUD も用意されています。レベルには荒地を含む地形と、隆起、坂、ループを含む簡単なレーストラックがあります。")
+
+LocalizedDisplayNames=(Language="zh-CN",Text="高级载具类游戏")
+LocalizedDescriptions=(Language="zh-CN",Text="高级载具类游戏模板含有一辆比普通载具类游戏模板更复杂的载具。您可通过键盘、手柄或触摸式设备上的虚拟摇杆来控制载具移动。该模板有2种摄像机视角——一种位于载具的后上方，而另一种位于载具内部。在使用内部视角时，您可通过鼠标、触摸屏或HMD设备观察四周。载具本身具有双横臂式悬架。该模板还包含引擎音效。载具内外还各有一个HUD，以显示当前装备和速度。关卡内含有一片粗糙地形，以及一条包含隆起、斜坡以及环道的简单道路。")
 
 ClassTypes=GameMode, WheeledVehicle, HUD, SpringArmComponent, Font, CameraComponent, TextRenderComponent, InputComponent, PhysicalMaterial, AudioComponent
 AssetTypes=Static Mesh, Landscape Layer, Material, Instance, Physical Material, Sound Attenuation, Sound cue, Sound Wave, Texture, Skeletal Mesh, Skeleton, Physics Asset, Animation blueprint, Data Asset, Vehicle Wheel
diff --git a/Templates/TP_VehicleAdvBP/manifest.json b/Templates/TP_VehicleAdvBP/manifest.json
--- a/Templates/TP_VehicleAdvBP/manifest.json
+++ b/Templates/TP_VehicleAdvBP/manifest.json
@@ -6,6 +6,14 @@
 			"Language": "en",
 			"Text": "Vehicle Advanced"
 		},
+		{
+			"Language": "ja",
+			"Text": "ビークル応用"
+		},
+		{
+			"Language": "zh-CN",
+			"Text": "高级载具类游戏"
+		},
 		{
 			"Language": "ko",
 			"Text": "비히클 고급"
@@ -17,9 +25,17 @@
 			"Language": "en",
 			"Text": "The vehicle advanced pack features a more complex vehicle than the regular vehicle template. Vehicle movement can be controlled with either keyboard, controller or virtual joystick on a touch device. The template features 2 camera views - One is positioned above and behind the vehicle, whilst the other is positioned inside. When using the inside view you can look around using the mouse or touchscreen or a HMD device. The vehicle itself features working double wishbone suspension. There is also an engine sound in this template. There is also a HUD both in and out of the car that shows the current gear and speed. The level contains a landscape with rough terrain as well as a simple track that features bumps, ramps and loop the loops."
 		},
+		{
+			"Language": "ja",
+			"Text": "ビークル応用テンプレートには通常のビークルテンプレートより複雑なビークルが含まれています。キーボード、コントローラー、タッチデバイスでのバーチャルジョイスティックでビークルを操作できます。このテンプレートには 2 方向からのカメラビューが用意されています。後方真上からのビューとビークル内部からのビューです。内部からのカメラを使用している時はマウス、タッチスクリーン、HMD デバイスで周りを見ることができます。このビークルにはダブルウィッシュボーン式サスペンションが含まれています。またエンジン音もテンプレートにあります。また現在のギアと速度を示す HUD も用意されています。レベルには荒地を含む地形と、隆起、坂、ループを含む簡単なレーストラックがあります。"
+		},
+		{
+			"Language": "zh-CN",
+			"Text": "高级载具类游戏模板含有一辆比普通载具类游戏模板更复杂的载具。您可通过键盘、手柄或触摸式设备上的虚拟摇杆来控制载具移动。该模板有2种摄像机视角——一种位于载具的后上方，而另一种位于载具内部。在使用内部视角时，您可通过鼠标、触摸屏或HMD设备观察四周。载具本身具有双横臂式悬架。该模板还包含引擎音效。载具内外还各有一个HUD，以显示当前装备和速度。关卡内含有一片粗糙地形，以及一条包含隆起、斜坡以及环道的简单道路。"
+		},
 		{
 			"Language": "ko",
-			"Text": "비히클 고급 팩은 일반 비히클 템플릿보다 복잡한 비히클이 특징입니다. 비히클 이동은 키보드, 콘트롤러, 터치 디바이스의 가상 조이스틱 등 어느 것으로도 제어 가능합니다. 템플릿에는 두 가지 카메라 뷰가 제공되는데, 하나는 비히클 상단 뒤쪽에, 다른 하나는 안쪽에 위치해 있습니다. 안쪽에 있을 때는 마우스나 터치스크린이나 HMD 디바이스를 사용하여 둘러볼 수 있습니다. 비히클 자체에는 더블 위시본 서스펜션은 물론 엔진 사운드도 구현되어 있습니다. 차량 안에도 밖에도 현재 기어와 속도를 나타내는 HUD 가 있습니다. 레벨에는 거친 지형의 랜드스케이프는 물론, 방지턱, 경사로, 루프로 된 간단한 트랙도 포함되어 있습니다."
+			"Text": "비히클 고급 템플릿은 일반 비히클 템플릿보다 복잡한 비히클이 특징입니다. 비히클 이동은 키보드, 컨트롤러, 터치 디바이스의 가상 조이스틱 등 어느 것으로도 제어 가능합니다. 템플릿에는 두 가지 카메라 뷰가 제공되는데, 하나는 비히클 상단 뒤쪽에, 다른 하나는 안쪽에 위치해 있습니다. 안쪽에 있을 때는 마우스나 터치스크린이나 HMD 디바이스를 사용하여 둘러볼 수 있습니다. 비히클 자체에는 더블 위시본 서스펜션은 물론 엔진 사운드도 구현되어 있습니다. 차량 안에도 밖에도 현재 기어와 속도를 나타내는 HUD 가 있습니다. 레벨에는 거친 지형의 랜드스케이프는 물론, 방지턱, 경사로, 간단한 루프 트랙도 포함되어 있습니다."
 		}
 	],
 	"AssetTypes":
@@ -27,10 +43,6 @@
 		{
 			"Language": "en",
 			"Text": "Static Mesh, Landscape Layer, Material, Instance, Physical Material, Sound Attenuation, Sound cue, Sound Wave, Texture, Skeletal Mesh, Skeleton, Physics Asset, Animation blueprint, Data Asset, Vehicle Wheel"
-		},
-		{
-			"Language": "ko",
-			"Text": "스태틱 메시, 랜드스케이프 레이어, 머티리얼, 인스턴스, 피지컬 머티리얼, 사운드 감쇠, 사운드 큐, 사운드 웨이브, 텍스처, 스켈레탈 메시, 스켈레톤, 피직스 애셋, 애니메이션 블루프린트, 데이터 애셋, 비히클 휠"
 		}
 	],	
 	"SearchTags":
@@ -38,10 +50,6 @@
 		{
 			"Language": "en",
 			"Text": "Tyre,Suspension,Wheel,Racing,Vehicle,Car"
-		},
-		{
-			"Language": "ko",
-			"Text": "타이어,서스펜션,휠,레이싱,비히클,차,경주"
 		}
 	],
 	"ClassTypes": "GameMode, WheeledVehicle, HUD, SpringArmComponent, Font, CameraComponent, TextRenderComponent, InputComponent, PhysicalMaterial, AudioComponent",	
diff --git a/Templates/TP_VehicleBP/Config/TemplateDefs.ini b/Templates/TP_VehicleBP/Config/TemplateDefs.ini
--- a/Templates/TP_VehicleBP/Config/TemplateDefs.ini
+++ b/Templates/TP_VehicleBP/Config/TemplateDefs.ini
@@ -3,7 +3,13 @@ LocalizedDisplayNames=(Language="en",Text="Vehicle")
 LocalizedDescriptions=(Language="en",Text="The vehicle template features a simple physics-driven vehicle. Vehicle movement can be controlled with either keyboard, controller or virtual joystick on a touch device. The template features 2 camera views - One is positioned above and behind the vehicle, whilst the other is positioned inside. When using the inside view the mouse or touchscreen or a HMD device. There is also a HUD both in and out of the car that shows the current gear and speed. The level features a smooth ramp, a slightly irregluar one, some objects that can be knocked over and a post which cannot.")
 
 LocalizedDisplayNames=(Language="ko",Text="비히클")
-LocalizedDescriptions=(Language="ko",Text="비히클 템플릿은 단순한 물리 기반 비히클을 특징으로 합니다. 비히클 이동은 키보드, 콘트롤러, 터치 디바이스의 가상 조이스틱 등 어느 것으로도 제어 가능합니다. 템플릿에는 두 가지 카메라 뷰가 제공되는데, 하나는 비히클 상단 뒤쪽에, 다른 하나는 안쪽에 위치해 있습니다. 안쪽에 있을 때는 마우스나 터치스크린이나 HMD 디바이스를 사용하여 둘러볼 수 있습니다. 차량 안에도 밖에도 현재 기어와 속도를 나타내는 HUD 가 있습니다. 레벨에는 약간 불규칙한 부드러운 경사로, 넘어뜨릴 수 있는 오브젝트와 넘어뜨릴 수 없는 기둥이 있습니다.")
+LocalizedDescriptions=(Language="ko",Text="비히클 템플릿은 단순한 물리 기반 비히클을 특징으로 합니다. 비히클 이동은 키보드, 컨트롤러, 터치 디바이스의 가상 조이스틱 등 어느 것으로도 제어 가능합니다. 템플릿에는 두 가지 카메라 뷰가 제공되는데, 하나는 비히클 상단 뒤쪽에, 다른 하나는 안쪽에 위치해 있습니다. 안쪽에 있을 때는 마우스나 터치스크린이나 HMD 디바이스를 사용하여 주변을 둘러볼 수 있습니다. 차량 안에도 밖에도 현재 기어와 속도를 나타내는 HUD 가 있습니다. 레벨에는 약간 불규칙한 부드러운 경사로, 넘어뜨릴 수 있는 오브젝트와 넘어뜨릴 수 없는 기둥이 있습니다.")
+
+LocalizedDisplayNames=(Language="ja",Text="ビークル")
+LocalizedDescriptions=(Language="ja",Text="ビークルテンプレートにはシンプルな物理駆動のビークルが含まれています。キーボード、コントローラー、タッチデバイスでのバーチャルジョイスティックでビークルを操作できます。このテンプレートには 2 方向からのカメラビューが用意されています。後方真上からのビューとビークル内部からのビューです。内部からのカメラを使用している時はマウス、タッチスクリーン、HMD デバイスで周りを見ることができます。また現在のギアと速度を示す HUD も用意されています。レベルにはスムーズな坂、多少不規則な坂、倒すことのできるオブジェクトとそうではない柱があります。")
+
+LocalizedDisplayNames=(Language="zh-CN",Text="载具类游戏")
+LocalizedDescriptions=(Language="zh-CN",Text="载具类游戏模板拥有一辆简单的物理驱动载具。您可通过键盘、手柄或触摸式设备上的虚拟摇杆来控制载具移动。该模板有2种摄像机视角——一种位于载具的后上方，而另一种位于载具内部。在使用内部视角时，您可通过鼠标、触摸屏或HMD设备观察四周。。载具内外还各有一个HUD，以显示当前装备和速度。关卡内含有一个光滑的斜坡和一个略不规则的斜坡，一些可以被撞飞的对象以及无法被撞开的柱子。")
 
 ClassTypes=GameMode, WheeledVehicle, HUD, Font, SpringArmComponent, CameraComponent, TextRenderComponent, InputComponent, VehicleWheel
 AssetTypes=Static Mesh, Animation Blueprint, Physics Asset, Physics Material, Skelton, Skeletal Mesh, Material, Material instance, Texture
diff --git a/Templates/TP_VehicleBP/manifest.json b/Templates/TP_VehicleBP/manifest.json
--- a/Templates/TP_VehicleBP/manifest.json
+++ b/Templates/TP_VehicleBP/manifest.json
@@ -10,6 +10,10 @@
 			"Language": "ko",
 			"Text": "비히클"
 		},
+		{
+			"Language": "zh-CN",
+			"Text": "载具类游戏"
+		},
 		{
 			"Language": "ja",
 			"Text": "レース"
@@ -23,11 +27,15 @@
 		},
 		{
 			"Language": "ko",
-			"Text": "비히클 팩은 단순한 물리 기반 비히클을 특징으로 합니다. 비히클 이동은 키보드, 콘트롤러, 터치 디바이스의 가상 조이스틱 등 어느 것으로도 제어 가능합니다. 템플릿에는 두 가지 카메라 뷰가 제공되는데, 하나는 비히클 상단 뒤쪽에, 다른 하나는 안쪽에 위치해 있습니다. 안쪽에 있을 때는 마우스나 터치스크린이나 HMD 디바이스를 사용하여 둘러볼 수 있습니다. 차량 안에도 밖에도 현재 기어와 속도를 나타내는 HUD 가 있습니다. 레벨에는 약간 불규칙한 부드러운 경사로, 넘어뜨릴 수 있는 오브젝트와 넘어뜨릴 수 없는 기둥이 있습니다."
+			"Text": "비히클 템플릿은 단순한 물리 기반 비히클을 특징으로 합니다. 비히클 이동은 키보드, 컨트롤러, 터치 디바이스의 가상 조이스틱 등 어느 것으로도 제어 가능합니다. 템플릿에는 두 가지 카메라 뷰가 제공되는데, 하나는 비히클 상단 뒤쪽에, 다른 하나는 안쪽에 위치해 있습니다. 안쪽에 있을 때는 마우스나 터치스크린이나 HMD 디바이스를 사용하여 주변을 둘러볼 수 있습니다. 차량 안에도 밖에도 현재 기어와 속도를 나타내는 HUD 가 있습니다. 레벨에는 약간 불규칙한 부드러운 경사로, 넘어뜨릴 수 있는 오브젝트와 넘어뜨릴 수 없는 기둥이 있습니다."
+		},
+		{
+			"Language": "zh-CN",
+			"Text": "载具类游戏模板拥有一辆简单的物理驱动载具。您可通过键盘、手柄或触摸式设备上的虚拟摇杆来控制载具移动。该模板有2种摄像机视角——一种位于载具的后上方，而另一种位于载具内部。在使用内部视角时，您可通过鼠标、触摸屏或HMD设备观察四周。。载具内外还各有一个HUD，以显示当前装备和速度。关卡内含有一个光滑的斜坡和一个略不规则的斜坡，一些可以被撞飞的对象以及无法被撞开的柱子。"
 		},
 		{
 			"Language": "ja",
-			"Text": "The vehicle pack features a simple physics-driven vehicle. Vehicle movement can be controlled with either keyboard, controller or virtual joystick on a touch device. The template features 2 camera views - On is positioned above and behind the vehicle, whilst the other is positioned inside. When using the inside view the mouse or touchscreen or a HMD device. There is also a HUD both in and out of the car that shows the current gear and speed. The level features a smooth ramp, a slightly irregluar one, some objects that can be knocked over and a post which cannot."
+			"Text": "ビークルテンプレートにはシンプルな物理駆動のビークルが含まれています。キーボード、コントローラー、タッチデバイスでのバーチャルジョイスティックでビークルを操作できます。このテンプレートには 2 方向からのカメラビューが用意されています。後方真上からのビューとビークル内部からのビューです。内部からのカメラを使用している時はマウス、タッチスクリーン、HMD デバイスで周りを見ることができます。また現在のギアと速度を示す HUD も用意されています。レベルにはスムーズな坂、多少不規則な坂、倒すことのできるオブジェクトとそうではない柱があります。"
 		}
 	],
 	"AssetTypes":
@@ -35,14 +43,6 @@
 		{
 			"Language": "en",
 			"Text": "Static Mesh, Animation Blueprint, Physics Asset, Physics Material, Skelton, Skeletal Mesh, Material, Material instance, Texture"
-		},
-		{
-			"Language": "ko",
-			"Text": "스태틱 메시, 애니메이션 블루프린트, 피직스 애셋, 피직스 머티리얼, 스켈레톤, 스켈레탈 메시, 머티리얼, 머티리얼 인스턴스, 텍스처"
-		},
-		{
-			"Language": "ja",
-			"Text": "スタティックメッシュ・アニメーションブループリント・フィジックス・フィジカルマテリアル・スケルトン・スケルタルメッシュ・マテリアル・マテリアルインスタンス・テクスチャ"
 		}
 	],
 	"SearchTags":
@@ -50,10 +50,6 @@
 		{
 			"Language": "en",
 			"Text": "Tyre,Wheel,Racing,Vehicle,Car"
-		},
-		{
-			"Language": "ko",
-			"Text": "타이어,휠,레이싱,비히클,차,경주"
 		}
 	],
 	"ClassTypes": "GameMode, WheeledVehicle, HUD, Font, SpringArmComponent, CameraComponent, TextRenderComponent, InputComponent, VehicleWheel",
diff --git a/Templates/TP_VirtualRealityBP/Config/TemplateDefs.ini b/Templates/TP_VirtualRealityBP/Config/TemplateDefs.ini
--- a/Templates/TP_VirtualRealityBP/Config/TemplateDefs.ini
+++ b/Templates/TP_VirtualRealityBP/Config/TemplateDefs.ini
@@ -3,7 +3,13 @@ LocalizedDisplayNames=(Language="en",Text="Virtual Reality")
 LocalizedDescriptions=(Language="en",Text="Blueprint Virtual Reality Template for Desktop and Playstation VR. The template features two modes of teleport locomotion with separate pawns and maps for each. The MotionControllerMap lets players teleport around using motion controllers. The HMDLocomotionMap lets players teleport using either a Gamepad or point-and-teleport using just the HeadMountedDisplay.")
 
 LocalizedDisplayNames=(Language="ko",Text="가상 현실")
-LocalizedDescriptions=(Language="ko",Text="데스크탑 및 Playstation VR 용 블루프린트 가상 현실 템플릿입니다. 각 플랫폼 별 전용 폰과 맵은 물론 두 가지 텔레포트 이동 방식이 포함되어 있습니다. 플레이어는 MotionControllerMap 에서 모션 컨트롤러를 사용하여 텔레포트할 수 있으며, HMDLocomotionMap 에서 게임패드 아니면 HeadMountedDisplay 만으로 이동할 곳을 지정하는 방식으로 텔레포트할 수 있습니다.")
+LocalizedDescriptions=(Language="ko",Text="데스크톱 및 Playstation VR 용 블루프린트 가상 현실 템플릿입니다. 각 플랫폼 별 전용 폰과 맵은 물론 두 가지 텔레포트 이동 방식이 포함되어 있습니다. 플레이어는 MotionControllerMap 에서 모션 컨트롤러를 사용하여 텔레포트할 수 있으며, HMDLocomotionMap 에서 게임패드 아니면 HeadMountedDisplay 만으로 이동할 곳을 지정하는 방식으로 텔레포트할 수 있습니다.")
+
+LocalizedDisplayNames=(Language="ja",Text="バーチャルリアリティ")
+LocalizedDescriptions=(Language="ja",Text="デスクトップ VR と Playstation VR 向け、ブループリントのバーチャルリアリティテンプレートです。テンプレートには二種類のテレポートロコモーションモードがあります。それぞれに ポーンとマップが用意されています。MotionControllerMap ではモーションコントローラーを使ったテレポートが可能です。HMDLocomotionMap ではゲームパッドを使ったテレポート、及び、ヘッドマウントディスプレイを使ったポイントアンドテレポートが可能です。")
+
+LocalizedDisplayNames=(Language="zh-CN",Text="虚拟现实应用")
+LocalizedDescriptions=(Language="zh-CN",Text="蓝图虚拟现实应用模板面向台式PC以及Playstation VR设备。模板包含2中传送移动模式，分别拥有独立的Pawn和地图。MotionControllerMap可以让玩家使用体感手柄四处传送，而HMDLocomotionMap可以让玩家使用游戏手柄或头显设备进行“指向-传送”式的传送。")
 
 ClassTypes=Character, GameMode, HUD, Actor
 AssetTypes=Animation Sequence, Animation Montage, Sound Wave, Skeletal Mesh, Animation Blueprint, Skeleton, Static Mesh, Texture.
diff --git a/Templates/TP_VirtualRealityBP/manifest.json b/Templates/TP_VirtualRealityBP/manifest.json
--- a/Templates/TP_VirtualRealityBP/manifest.json
+++ b/Templates/TP_VirtualRealityBP/manifest.json
@@ -4,12 +4,36 @@
     {
       "Language": "en",
       "Text": "Virtual Reality"
+    },
+    {
+      "Language": "ko",
+      "Text": "가상 현실"
+    },
+    {
+      "Language": "ja",
+      "Text": "バーチャルリアリティ"
+    },
+    {
+      "Language": "zh-CN",
+      "Text": "虚拟现实应用"
     }
   ],
   "Description": [
     {
       "Language": "en",
       "Text": "Blueprint Virtual Reality feature pack for Desktop and Playstation VR. The pack features two modes of teleport locomotion with separate pawns and maps for each. The MotionControllerMap lets players teleport around using motion controllers. The HMDLocomotionMap lets players teleport using either a Gamepad or point-and-teleport using just the HeadMountedDisplay."
+    },
+    {
+      "Language": "ko",
+      "Text": "데스크톱 및 Playstation VR 용 블루프린트 가상 현실 템플릿입니다. 각 플랫폼 별 전용 폰과 맵은 물론 두 가지 텔레포트 이동 방식이 포함되어 있습니다. 플레이어는 MotionControllerMap 에서 모션 컨트롤러를 사용하여 텔레포트할 수 있으며, HMDLocomotionMap 에서 게임패드 아니면 HeadMountedDisplay 만으로 이동할 곳을 지정하는 방식으로 텔레포트할 수 있습니다."
+    },
+    {
+      "Language": "ja",
+      "Text": "デスクトップ VR と Playstation VR 向け、ブループリントのバーチャルリアリティテンプレートです。テンプレートには二種類のテレポートロコモーションモードがあります。それぞれに ポーンとマップが用意されています。MotionControllerMap ではモーションコントローラーを使ったテレポートが可能です。HMDLocomotionMap ではゲームパッドを使ったテレポート、及び、ヘッドマウントディスプレイを使ったポイントアンドテレポートが可能です。"
+    },
+    {
+      "Language": "zh-CN",
+      "Text": "蓝图虚拟现实应用模板面向台式PC以及Playstation VR设备。模板包含2中传送移动模式，分别拥有独立的Pawn和地图。MotionControllerMap可以让玩家使用体感手柄四处传送，而HMDLocomotionMap可以让玩家使用游戏手柄或头显设备进行“指向-传送”式的传送。"
     }
   ],
   "AssetTypes": [
diff --git a/Templates/TP_nDisplay/Config/TemplateDefs.ini b/Templates/TP_nDisplay/Config/TemplateDefs.ini
--- a/Templates/TP_nDisplay/Config/TemplateDefs.ini
+++ b/Templates/TP_nDisplay/Config/TemplateDefs.ini
@@ -2,6 +2,15 @@
 LocalizedDisplayNames=(Language="en",Text="nDisplay")
 LocalizedDescriptions=(Language="en",Text="This project template provides multi-display functionality using a cluster of PCs. Use as a starting point for stereoscropic/monoscopic Caves and PowerWall setups.")
 
+LocalizedDisplayNames=(Language="ko",Text="nDisplay")
+LocalizedDescriptions=(Language="ko",Text="PC 클러스터에 연결된 멀티 디스플레이 기능을 제공하는 프로젝트 템플릿입니다. 스테레오/모노스코픽 동굴이나 파워 월 구성을 위한 시작점으로 사용할 수 있습니다.")
+
+LocalizedDisplayNames=(Language="ja",Text="nDisplay")
+LocalizedDescriptions=(Language="ja",Text="PCクラスタを使ったマルチディスプレイ機能のテンプレートです。立体視/平面視の Cave や PowerWall のセットアップでの土台として使用してください。")
+
+LocalizedDisplayNames=(Language="zh-CN",Text="nDisplay")
+LocalizedDescriptions=(Language="zh-CN",Text="此项目模板提供了使用PC集群的多显示设备功能。它可作为立体/平面Caves以及PowerWall设置的基础。")
+
 ClassTypes=GameMode, Actor
 AssetTypes=Static Mesh, Texture
 
diff --git a/Templates/TP_nDisplayBP/Config/TemplateDefs.ini b/Templates/TP_nDisplayBP/Config/TemplateDefs.ini
--- a/Templates/TP_nDisplayBP/Config/TemplateDefs.ini
+++ b/Templates/TP_nDisplayBP/Config/TemplateDefs.ini
@@ -2,6 +2,15 @@
 LocalizedDisplayNames=(Language="en",Text="nDisplay")
 LocalizedDescriptions=(Language="en",Text="This project template provides multi-display functionality using a cluster of PCs. Use as a starting point for stereoscropic/monoscopic Caves and PowerWall setups.")
 
+LocalizedDisplayNames=(Language="ko",Text="nDisplay")
+LocalizedDescriptions=(Language="ko",Text="PC 클러스터에 연결된 멀티 디스플레이 기능을 제공하는 프로젝트 템플릿입니다. 스테레오/모노스코픽 동굴이나 파워 월 구성을 위한 시작점으로 사용할 수 있습니다.")
+
+LocalizedDisplayNames=(Language="ja",Text="nDisplay")
+LocalizedDescriptions=(Language="ja",Text="PCクラスタを使ったマルチディスプレイ機能のテンプレートです。立体視/平面視の Cave や PowerWall のセットアップでの土台として使用してください。")
+
+LocalizedDisplayNames=(Language="zh-CN",Text="nDisplay")
+LocalizedDescriptions=(Language="zh-CN",Text="此项目模板提供了使用PC集群的多显示设备功能。它可作为立体/平面Caves以及PowerWall设置的基础。")
+
 ClassTypes=GameMode, Actor
 AssetTypes=Static Mesh, Texture
 
