Moving map for Vehicle CPP to the correct folder and updating manifest files with the proper shared packs for feature pack generation.

diff --git a/Templates/TP_VehicleAdv/manifest.json b/Templates/TP_VehicleAdv/manifest.json
--- a/Templates/TP_VehicleAdv/manifest.json
+++ b/Templates/TP_VehicleAdv/manifest.json
@@ -64,17 +64,23 @@
 	"AdditionalFeaturePacks":
 	[
 		{
-			"MountName": "VehicleAdv",
+			"MountName": "Vehicles",
 			"DetailLevels":
 			[
 				"Standard"
 			]
 		},
 		{
-			"MountName": "Geometry",
+			"MountName": "Track",
+			"DetailLevels":
+			[
+				"High"
+			]
+		},
+		{
+			"MountName": "LevelPrototyping",
 			"DetailLevels":
 			[
-				"Standard",
 				"High"
 			]
 		}
diff --git a/Templates/TP_VehicleAdvBP/manifest.json b/Templates/TP_VehicleAdvBP/manifest.json
--- a/Templates/TP_VehicleAdvBP/manifest.json
+++ b/Templates/TP_VehicleAdvBP/manifest.json
@@ -64,17 +64,23 @@
 	"AdditionalFeaturePacks":
 	[
 		{
-			"MountName": "VehicleAdv",
+			"MountName": "Vehicles",
 			"DetailLevels":
 			[
 				"Standard"
 			]
 		},
 		{
-			"MountName": "Geometry",
+			"MountName": "Track",
+			"DetailLevels":
+			[
+				"High"
+			]
+		},
+		{
+			"MountName": "LevelPrototyping",
 			"DetailLevels":
 			[
-				"Standard",
 				"High"
 			]
 		}
