[Feature Packs] Created Feature Packs for 1st person , 3rd person, top down, vehicle, vr templates (including needed changes in manifests and making sure world settings contain game mode)

#content only

diff --git a/Templates/TP_FirstPerson/manifest.json b/Templates/TP_FirstPerson/manifest.json
--- a/Templates/TP_FirstPerson/manifest.json
+++ b/Templates/TP_FirstPerson/manifest.json
@@ -1,7 +1,6 @@
 {
 	"Version": 1,
-	"Name":
-	[
+	"Name": [
 		{
 			"Language": "en",
 			"Text": "First Person"
@@ -19,8 +18,7 @@
 			"Text": "ファーストパーソン"
 		}
 	],
-	"Description":
-	[
+	"Description": [
 		{
 			"Language": "en",
 			"Text": "The first person template features a player character (represented by a pair of arms) which is viewed from first person perspective. The character can be moved around the level using keyboard, controller or virtual joystick on a touch device. Additionally the player can look around using mouse, controller or virtual joystick on a touch device. The character can also pick up a gun that, using mouse, controller or virtual joystick on a touch device will fire a simple sphere projectile that will affect some physics objects in the level, whilst bounce off the arena walls."
@@ -38,15 +36,13 @@
 			"Text": "ファーストパーソン テンプレートは、一人称視点から見たプレイヤー キャラクター (両腕が表示される) を特徴とします。タッチ デバイスでキーボード、コントローラー、またはバーチャル ジョイスティックを使用して、レベル内でキャラクターを移動させることができます。さらにプレイヤーは、タッチデバイスでマウス、コントローラー、またはバーチャル ジョイスティックを使って見回すことができます。キャラクターは、タッチデバイスで、マウス、コントローラー、またはバーチャル ジョイスティックを使用して、アリーナの壁で跳ね返りながら、レベル内の一部の物理オブジェクトに影響を与えるシンプルな球体の発射体を発射する銃を拾うこともできます。"
 		}
 	],
-	"AssetTypes":
-	[
+	"AssetTypes": [
 		{
 			"Language": "en",
 			"Text": "Animation Sequence, Animation Montage, Sound Wave, Skeletal Mesh, Animation Blueprint, Skeleton, Static Mesh."
 		}
-	],	
-	"SearchTags":
-	[
+	],
+	"SearchTags": [
 		{
 			"Language": "en",
 			"Text": "First Person,Shoot,Fire,Aim"
@@ -54,14 +50,13 @@
 	],
 	"ClassTypes": "Actor, Character, GameMode",
 	"Category": "CodeFeature",
-	"FocusAsset": "/Game/FirstPerson/Blueprints/FirstPersonCharacter.FirstPersonCharacter",
+	"FocusAsset": "/Game/FirstPerson/Blueprints/BP_FirstPersonCharacter.BP_FirstPersonCharacter",
 	"Thumbnail": "TP_FirstPerson.png",
-	"Screenshots":
-	[
+	"Screenshots": [
 		"TP_FirstPerson_Preview.png"
 	],
 	"AdditionalFeaturePacks": [
-				{
+		{
 			"MountName": "FPWeapon",
 			"DetailLevels": [
 				"Standard"
@@ -75,8 +70,7 @@
 		},
 		{
 			"MountName": "FirstPersonArms",
-			"DetailLevels":
-			[
+			"DetailLevels": [
 				"Standard"
 			]
 		}
diff --git a/Templates/TP_FirstPersonBP/manifest.json b/Templates/TP_FirstPersonBP/manifest.json
--- a/Templates/TP_FirstPersonBP/manifest.json
+++ b/Templates/TP_FirstPersonBP/manifest.json
@@ -1,7 +1,6 @@
 {
 	"Version": 1,
-	"Name":
-	[
+	"Name": [
 		{
 			"Language": "en",
 			"Text": "First Person"
@@ -19,8 +18,7 @@
 			"Text": "ファーストパーソン"
 		}
 	],
-	"Description":
-	[
+	"Description": [
 		{
 			"Language": "en",
 			"Text": "The first person template features a player character (represented by a pair of arms) which is viewed from first person perspective. The character can be moved around the level using keyboard, controller or virtual joystick on a touch device. Additionally the player can look around using mouse, controller or virtual joystick on a touch device. The character can also pick up a gun that, using mouse, controller or virtual joystick on a touch device will fire a simple sphere projectile that will affect some physics objects in the level, whilst bounce off the arena walls."
@@ -38,15 +36,13 @@
 			"Text": "ファーストパーソンテンプレートには、一人称視点で表示されるプレイヤーキャラクタ（腕として表示されています）がいます。キーボード、コントローラー、タッチデバイスでのバーチャルジョイスティックでこのキャラクタを移動させることができます。追加して、プレイヤーはマウス、コントローラー、タッチデバイスでのバーチャルジョイスティックで周りを見回すことができます。さらにキャラクタは銃を装備して、マウス、コントローラー、タッチデバイスでのバーチャルジョイスティックを使ってシンプルな球体の弾を発射できます。この弾は物理オブジェクトに影響し、アリーナの壁で反射します。"
 		}
 	],
-	"AssetTypes":
-	[
+	"AssetTypes": [
 		{
 			"Language": "en",
 			"Text": "Animation Sequence, Animation Montage, Sound Wave, Skeletal Mesh, Animation Blueprint, Skeleton, Static Mesh."
 		}
-	],	
-	"SearchTags":
-	[
+	],
+	"SearchTags": [
 		{
 			"Language": "en",
 			"Text": "First Person,Shoot,Fire,Aim"
@@ -54,10 +50,9 @@
 	],
 	"ClassTypes": "Actor, Character, GameMode",
 	"Category": "BlueprintFeature",
-	"FocusAsset": "/Game/FirstPersonBP/FirstPersonOverview.FirstPersonOverview",
+	"FocusAsset": "/Game/FirstPerson/Blueprints/BP_FirstPersonCharacter.BP_FirstPersonCharacter",
 	"Thumbnail": "TP_FirstPersonBP.png",
-	"Screenshots":
-	[
+	"Screenshots": [
 		"TP_FirstPersonBP_Preview.png"
 	],
 	"AdditionalFeaturePacks": [
@@ -75,8 +70,7 @@
 		},
 		{
 			"MountName": "FirstPersonArms",
-			"DetailLevels":
-			[
+			"DetailLevels": [
 				"Standard"
 			]
 		}
diff --git a/Templates/TP_ThirdPerson/manifest.json b/Templates/TP_ThirdPerson/manifest.json
--- a/Templates/TP_ThirdPerson/manifest.json
+++ b/Templates/TP_ThirdPerson/manifest.json
@@ -1,7 +1,6 @@
 {
 	"Version": 1,
-	"Name":
-	[
+	"Name": [
 		{
 			"Language": "en",
 			"Text": "Third Person"
@@ -19,8 +18,7 @@
 			"Text": "サードパーソン"
 		}
 	],
-	"Description":
-	[
+	"Description": [
 		{
 			"Language": "en",
 			"Text": "The third person pack features a playable character where the camera is positioned behind and slightly above the character. As the character moves using mouse, keyboard, controller or virtual joystick on a touch device, the camera follows the character, observing from an over-the-shoulder point of view. This perspective emphasizes the main character and is popular in action and adventure games. The template map contains several objects that the player can jump upon, and push around."
@@ -38,15 +36,13 @@
 			"Text": "サードパーソン パックは、後ろ側と少し上にカメラが配置されているプレイ可能なキャラクターを備えています。タッチデバイスでマウス、キーボード、コントローラー、またはバーチャル ジョイスティックの使用によりキャラクターが移動すると、カメラはキャラクターを追跡し、肩越しの視点から観察します。この視点では主人公が強調されるため、アクションゲームやアドベンチャー ゲームで人気があります。テンプレート マップには、プレイヤーがジャンプしたり押すことができるオブジェクトがいくつか含まれています。"
 		}
 	],
-	"AssetTypes":
-	[
+	"AssetTypes": [
 		{
 			"Language": "en",
 			"Text": "Animation Sequence, Animation Blueprint, Skeleton, Skeletal Mesh, Control Rig, IK Rig, Retargeter"
 		}
-	],	
-	"SearchTags":
-	[
+	],
+	"SearchTags": [
 		{
 			"Language": "en",
 			"Text": "Action, Adventure, Third"
@@ -54,25 +50,21 @@
 	],
 	"ClassTypes": "Character, SpringArmComponent, GameMode",
 	"Category": "CodeFeature",
-	"FocusAsset": "/Game/ThirdPerson/Blueprints",
+	"FocusAsset": "/Game/ThirdPerson/Blueprints/BP_ThirdPersonCharacter.BP_ThirdPersonCharacter",
 	"Thumbnail": "TP_ThirdPerson.png",
-	"Screenshots":
-	[
+	"Screenshots": [
 		"TP_ThirdPerson_Preview.png"
 	],
-	"AdditionalFeaturePacks":
-	[		
+	"AdditionalFeaturePacks": [
 		{
 			"MountName": "LevelPrototyping",
-			"DetailLevels":
-			[			
+			"DetailLevels": [
 				"High"
 			]
 		},
 		{
 			"MountName": "Characters",
-			"DetailLevels":
-			[		
+			"DetailLevels": [
 				"High"
 			]
 		}
diff --git a/Templates/TP_ThirdPersonBP/manifest.json b/Templates/TP_ThirdPersonBP/manifest.json
--- a/Templates/TP_ThirdPersonBP/manifest.json
+++ b/Templates/TP_ThirdPersonBP/manifest.json
@@ -1,7 +1,6 @@
 {
 	"Version": 1,
-	"Name":
-	[
+	"Name": [
 		{
 			"Language": "en",
 			"Text": "Third Person"
@@ -19,8 +18,7 @@
 			"Text": "サードパーソン"
 		}
 	],
-	"Description":
-	[
+	"Description": [
 		{
 			"Language": "en",
 			"Text": "The third person pack features a playable character where the camera is positioned behind and slightly above the character. As the character moves using mouse, keyboard, controller or virtual joystick on a touch device, the camera follows the character, observing from an over-the-shoulder point of view. This perspective emphasizes the main character and is popular in action and adventure games. The template map contains several objects that the player can jump upon, and push around."
@@ -38,15 +36,13 @@
 			"Text": "サードパーソン パックは、後ろ側と少し上にカメラが配置されているプレイ可能なキャラクターを備えています。タッチデバイスでマウス、キーボード、コントローラー、またはバーチャル ジョイスティックの使用によりキャラクターが移動すると、カメラはキャラクターを追跡し、肩越しの視点から観察します。この視点では主人公が強調されるため、アクションゲームやアドベンチャー ゲームで人気があります。テンプレート マップには、プレイヤーがジャンプしたり押すことができるオブジェクトがいくつか含まれています。"
 		}
 	],
-	"AssetTypes":
-	[
+	"AssetTypes": [
 		{
 			"Language": "en",
 			"Text": "Animation Sequence, Animation Blueprint, Skeleton, Skeletal Mesh, Control Rig, IK Rig, Retargeter"
 		}
-	],	
-	"SearchTags":
-	[
+	],
+	"SearchTags": [
 		{
 			"Language": "en",
 			"Text": "Action, Adventure, Third"
@@ -54,25 +50,21 @@
 	],
 	"ClassTypes": "Character, SpringArmComponent, GameMode",
 	"Category": "BlueprintFeature",
-	"FocusAsset": "/Game/ThirdPerson/Blueprints",
+	"FocusAsset": "/Game/ThirdPerson/Blueprints/BP_ThirdPersonCharacter.BP_ThirdPersonCharacter",
 	"Thumbnail": "TP_ThirdPersonBP.png",
-	"Screenshots":
-	[
+	"Screenshots": [
 		"TP_ThirdPersonBP_Preview.png"
 	],
-	"AdditionalFeaturePacks":
-	[
+	"AdditionalFeaturePacks": [
 		{
 			"MountName": "LevelPrototyping",
-			"DetailLevels":
-			[
+			"DetailLevels": [
 				"High"
 			]
 		},
 		{
 			"MountName": "Characters",
-			"DetailLevels":
-			[
+			"DetailLevels": [
 				"High"
 			]
 		}
diff --git a/Templates/TP_TopDown/manifest.json b/Templates/TP_TopDown/manifest.json
--- a/Templates/TP_TopDown/manifest.json
+++ b/Templates/TP_TopDown/manifest.json
@@ -1,7 +1,6 @@
 {
 	"Version": 1,
-	"Name":
-	[
+	"Name": [
 		{
 			"Language": "en",
 			"Text": "Top Down"
@@ -19,8 +18,7 @@
 			"Text": "トップダウン"
 		}
 	],
-	"Description":
-	[
+	"Description": [
 		{
 			"Language": "en",
 			"Text": "The top down pack features a playable character where the camera is positioned behind and above the avatar at greater distance than in third person perspective. The player is controlled by using the mouse or touchscreen to click on the required destination and uses the navigation system to facilitate the character movement. This perspective is often used in action role playing games."
@@ -38,15 +36,13 @@
 			"Text": "トップダウンテンプレートにはプレイアブルキャラクタが存在し、カメラはキャラクタの背後、サードパーソンテンプレートより遠い上方にあります。マウスやタッチスクリーンで行き先をクリックすることでキャラクタを操作することができます。キャラクタはナビゲーションシステムを利用して移動します。この視点はアクションロールプレイングゲームでよく使用されています。"
 		}
 	],
-	"AssetTypes":
-	[
+	"AssetTypes": [
 		{
 			"Language": "en",
 			"Text": "Animation Sequence, Skeletal Mesh, Animation Blueprint,Skeleton, Blend Space 1D"
 		}
-	],	
-	"SearchTags":
-	[
+	],
+	"SearchTags": [
 		{
 			"Language": "en",
 			"Text": "Role, RPG,Role Playing,Top Down"
@@ -54,32 +50,27 @@
 	],
 	"ClassTypes": "GameMode, Character, CameraComponent, SpringArmComponent, PlayerController",
 	"Category": "CodeFeature",
-	"FocusAsset": "/Game/TopDown/Blueprints",
+	"FocusAsset": "/Game/TopDown/Blueprints/BP_TopDownCharacter.BP_TopDownCharacter",
 	"Thumbnail": "TP_TopDown.png",
-	"Screenshots":
-	[
+	"Screenshots": [
 		"TP_TopDown_Preview.png"
 	],
-	"AdditionalFeaturePacks":
-	[
+	"AdditionalFeaturePacks": [
 		{
 			"MountName": "LevelPrototyping",
-			"DetailLevels":
-			[
+			"DetailLevels": [
 				"High"
 			]
-		},				
-{
+		},
+		{
 			"MountName": "Cursor",
-			"DetailLevels":
-			[
+			"DetailLevels": [
 				"Standard"
 			]
 		},
 		{
 			"MountName": "Characters",
-			"DetailLevels":
-			[
+			"DetailLevels": [
 				"High"
 			]
 		}
diff --git a/Templates/TP_TopDownBP/manifest.json b/Templates/TP_TopDownBP/manifest.json
--- a/Templates/TP_TopDownBP/manifest.json
+++ b/Templates/TP_TopDownBP/manifest.json
@@ -1,7 +1,6 @@
 {
 	"Version": 1,
-	"Name":
-	[
+	"Name": [
 		{
 			"Language": "en",
 			"Text": "Top Down"
@@ -19,8 +18,7 @@
 			"Text": "トップダウン"
 		}
 	],
-	"Description":
-	[
+	"Description": [
 		{
 			"Language": "en",
 			"Text": "The top down template features a playable character where the camera is positioned behind and above the avatar at greater distance than in third person perspective. The player is controlled by using the mouse, gamepad, or touchscreen and moves the character by clicking on a destination, or following the cursor if input is held. This perspective is often used in action role playing games."
@@ -38,15 +36,13 @@
 			"Text": "トップダウン テンプレートは、カメラがアバターの後ろと上に配置されている (サードパーソン視点よりも距離が遠い位置) プレイ可能なキャラクターを特徴としています。プレイヤーは、マウス、ゲームパッド、またはタッチ スクリーンの使用により制御され、目的地をクリックするか、入力が保持されている場合はカーソルをたどることでキャラクターを移動させます。この視点は、アクション ロールプレイ ゲームでよく使用されます。"
 		}
 	],
-	"AssetTypes":
-	[
+	"AssetTypes": [
 		{
 			"Language": "en",
 			"Text": "Animation Sequence, Skeletal Mesh, Animation Blueprint, Skeleton"
 		}
-	],	
-	"SearchTags":
-	[
+	],
+	"SearchTags": [
 		{
 			"Language": "en",
 			"Text": "RPG,Role Playing,Top Down"
@@ -54,32 +50,27 @@
 	],
 	"ClassTypes": "Character, PlayerController, CameraComponent, SpringArmComponent, GameMode",
 	"Category": "BlueprintFeature",
-	"FocusAsset": "/Game/TopDown/Blueprints",
+	"FocusAsset": "/Game/TopDown/Blueprints/BP_TopDownCharacter.BP_TopDownCharacter",
 	"Thumbnail": "TP_TopDownBP.png",
-	"Screenshots":
-	[
+	"Screenshots": [
 		"TP_TopDownBP_Preview.png"
 	],
-	"AdditionalFeaturePacks":
-	[
+	"AdditionalFeaturePacks": [
 		{
 			"MountName": "LevelPrototyping",
-			"DetailLevels":
-			[			
+			"DetailLevels": [
 				"High"
 			]
 		},
 		{
 			"MountName": "Cursor",
-			"DetailLevels":
-			[			
+			"DetailLevels": [
 				"Standard"
 			]
 		},
 		{
 			"MountName": "Characters",
-			"DetailLevels":
-			[
+			"DetailLevels": [
 				"High"
 			]
 		}
diff --git a/Templates/TP_VehicleAdvBP/Config/DefaultEngine.ini b/Templates/TP_VehicleAdvBP/Config/DefaultEngine.ini
--- a/Templates/TP_VehicleAdvBP/Config/DefaultEngine.ini
+++ b/Templates/TP_VehicleAdvBP/Config/DefaultEngine.ini
@@ -15,7 +15,7 @@ ThreePlayerSplitscreenLayout=FavorTop
 GameInstanceClass=/Script/Engine.GameInstance
 GameDefaultMap=/Game/VehicleTemplate/Maps/VehicleExampleMap.VehicleExampleMap
 ServerDefaultMap=/Engine/Maps/Entry.Entry
-GlobalDefaultGameMode=/Game/VehicleTemplate/Blueprints/VehicleAdvGameMode.VehicleAdvGameMode_C
+GlobalDefaultGameMode=/Script/Engine.GameMode
 GlobalDefaultServerGameMode=None
 
 [/Script/Engine.RendererSettings]
@@ -55,3 +55,17 @@ CompressionQualityModifier=1.000000
 AutoStreamingThreshold=0.000000
 SoundCueCookQualityIndex=-1
 
+[/Script/AndroidFileServerEditor.AndroidFileServerRuntimeSettings]
+bEnablePlugin=True
+bAllowNetworkConnection=True
+SecurityToken=9523615841FA3CF96096AB95EECF5F7C
+bIncludeInShipping=False
+bAllowExternalStartInShipping=False
+bCompileAFSProject=False
+bUseCompression=False
+bLogFiles=False
+bReportStats=False
+ConnectionType=USBOnly
+bUseManualIPAddress=False
+ManualIPAddress=
+
diff --git a/Templates/TP_VehicleAdvBP/manifest.json b/Templates/TP_VehicleAdvBP/manifest.json
--- a/Templates/TP_VehicleAdvBP/manifest.json
+++ b/Templates/TP_VehicleAdvBP/manifest.json
@@ -1,7 +1,6 @@
 {
 	"Version": 1,
-	"Name":
-	[
+	"Name": [
 		{
 			"Language": "en",
 			"Text": "Vehicle"
@@ -19,8 +18,7 @@
 			"Text": "비히클"
 		}
 	],
-	"Description":
-	[
+	"Description": [
 		{
 			"Language": "en",
 			"Text": "The vehicle pack features both a regular vehicle and a more complex vehicle. Vehicle movement can be controlled with either keyboard, controller or virtual joystick on a touch device. The advanced vehicle itself features working double wishbone suspension. There is also an engine sound in this template. There is also a HUD that shows the current gear and speed."
@@ -38,49 +36,42 @@
 			"Text": "비히클 팩은 일반 비히클과 조금 더 복잡한 비히클을 모두 제공합니다. 비히클의 이동은 키보드, 컨트롤러 또는 터치 디바이스의 버추얼 조이스틱으로 제어 가능합니다. 고급 비히클은 더블 위시본 서스펜션을 갖추고 있습니다. 이 템플릿에는 엔진 사운드와, 현재 기어와 속도를 보여주는 HUD도 포함되어 있습니다."
 		}
 	],
-	"AssetTypes":
-	[
+	"AssetTypes": [
 		{
 			"Language": "en",
 			"Text": "Static Mesh, Landscape Layer, Material, Instance, Physical Material, Sound Attenuation, Sound cue, Sound Wave, Texture, Skeletal Mesh, Skeleton, Physics Asset, Animation blueprint, Data Asset, Vehicle Wheel"
 		}
-	],	
-	"SearchTags":
-	[
+	],
+	"SearchTags": [
 		{
 			"Language": "en",
 			"Text": "Tyre,Suspension,Wheel,Racing,Vehicle,Car"
 		}
 	],
-	"ClassTypes": "GameMode, WheeledVehicle, HUD, SpringArmComponent, Font, CameraComponent, TextRenderComponent, InputComponent, PhysicalMaterial, AudioComponent",	
+	"ClassTypes": "GameMode, WheeledVehicle, HUD, SpringArmComponent, Font, CameraComponent, TextRenderComponent, InputComponent, PhysicalMaterial, AudioComponent",
 	"Category": "BlueprintFeature",
-	"FocusAsset": "/Game/VehicleAdvBP/AdvancedVehicleOverview.AdvancedVehicleOverview",
+	"FocusAsset": "/Game/VehicleTemplate/Blueprints/VehiclePlayerController.VehiclePlayerController",
 	"Thumbnail": "TP_VehicleAdvBP.png",
-	"Screenshots":
-	[
+	"Screenshots": [
 		"TP_VehicleAdvBP_Preview.png"
 	],
 	"SortKey": "ZVehicle_AdvBP",
-	"AdditionalFeaturePacks":
-	[
+	"AdditionalFeaturePacks": [
 		{
 			"MountName": "Vehicles",
-			"DetailLevels":
-			[
+			"DetailLevels": [
 				"Standard"
 			]
 		},
 		{
 			"MountName": "Track",
-			"DetailLevels":
-			[
+			"DetailLevels": [
 				"High"
 			]
 		},
 		{
 			"MountName": "LevelPrototyping",
-			"DetailLevels":
-			[
+			"DetailLevels": [
 				"High"
 			]
 		}
diff --git a/Templates/TP_VirtualRealityBP/manifest.json b/Templates/TP_VirtualRealityBP/manifest.json
--- a/Templates/TP_VirtualRealityBP/manifest.json
+++ b/Templates/TP_VirtualRealityBP/manifest.json
@@ -67,18 +67,16 @@
 	],
 	"ClassTypes": "Pawn, GameMode, Blueprint Interface, Actor",
 	"Category": "BlueprintFeature",
-	"FocusAsset": "/Game/VRTemplate/Blueprints",
+	"FocusAsset": "/Game/VRTemplate/Blueprints/VRPawn.VRPawn",
 	"Thumbnail": "TP_VirtualRealityBP.png",
 	"Screenshots": [
 		"TP_VirtualRealityBP_Preview.png"
 	],
 	"SortKey": "VirtualRealityBP",
-	"AdditionalFeaturePacks": 
-	[
+	"AdditionalFeaturePacks": [
 		{
 			"MountName": "LevelPrototyping",
-			"DetailLevels": 
-			[
+			"DetailLevels": [
 				"High"
 			]
 		},
