Templates - updated FeaturePacks
#rnx

diff --git a/Templates/TP_FirstPersonBP/manifest.json b/Templates/TP_FirstPersonBP/manifest.json
--- a/Templates/TP_FirstPersonBP/manifest.json
+++ b/Templates/TP_FirstPersonBP/manifest.json
@@ -57,9 +57,9 @@
 	],
 	"AdditionalFeaturePacks": [
 		{
-			"MountName": "Weapons",
+			"MountName": "Input",
 			"DetailLevels": [
-				"Standard"
+				"High"
 			]
 		},
 		{
diff --git a/Templates/TP_ThirdPersonBP/manifest.json b/Templates/TP_ThirdPersonBP/manifest.json
--- a/Templates/TP_ThirdPersonBP/manifest.json
+++ b/Templates/TP_ThirdPersonBP/manifest.json
@@ -67,6 +67,12 @@
 			"DetailLevels": [
 				"High"
 			]
+		},
+		{
+			"MountName": "Input",
+			"DetailLevels": [
+				"High"
+			]
 		}
 	]
 }
diff --git a/Templates/TP_VehicleAdvBP/manifest.json b/Templates/TP_VehicleAdvBP/manifest.json
--- a/Templates/TP_VehicleAdvBP/manifest.json
+++ b/Templates/TP_VehicleAdvBP/manifest.json
@@ -64,13 +64,7 @@
 			]
 		},
 		{
-			"MountName": "Track",
-			"DetailLevels": [
-				"High"
-			]
-		},
-		{
-			"MountName": "LevelPrototyping",
+			"MountName": "Input",
 			"DetailLevels": [
 				"High"
 			]
